Fix for "//UE5/Main - NonUnity Compile UnrealEditor Win64 - ConsoleVariablesEditorStyle.cpp contains undeclared identifier errors."

diff --git a/Engine/Plugins/Experimental/ConsoleVariables/Source/ConsoleVariablesEditor/Private/ConsoleVariablesEditorStyle.cpp b/Engine/Plugins/Experimental/ConsoleVariables/Source/ConsoleVariablesEditor/Private/ConsoleVariablesEditorStyle.cpp
--- a/Engine/Plugins/Experimental/ConsoleVariables/Source/ConsoleVariablesEditor/Private/ConsoleVariablesEditorStyle.cpp
+++ b/Engine/Plugins/Experimental/ConsoleVariables/Source/ConsoleVariablesEditor/Private/ConsoleVariablesEditorStyle.cpp
@@ -3,8 +3,13 @@
 #include "ConsoleVariablesEditorStyle.h"
 
 #include "EditorStyleSet.h"
+#include "Framework/Application/SlateApplication.h"
 #include "Interfaces/IPluginManager.h"
+#include "Styling/CoreStyle.h"
 #include "Styling/SlateStyleRegistry.h"
+#include "Styling/SlateTypes.h"
+#include "Widgets/Input/SButton.h"
+#include "Widgets/Input/SComboButton.h"
 
 TSharedPtr<FSlateStyleSet> FConsoleVariablesEditorStyle::StyleInstance = nullptr;
 
