Fix for "//UE5/Main - NonUnity Compile UnrealEditor Win64 - SConsoleVariablesEditorMainPanel.cpp contains undeclared identifier errors."

diff --git a/Engine/Plugins/Experimental/ConsoleVariables/Source/ConsoleVariablesEditor/Private/Views/MainPanel/SConsoleVariablesEditorMainPanel.cpp b/Engine/Plugins/Experimental/ConsoleVariables/Source/ConsoleVariablesEditor/Private/Views/MainPanel/SConsoleVariablesEditorMainPanel.cpp
--- a/Engine/Plugins/Experimental/ConsoleVariables/Source/ConsoleVariablesEditor/Private/Views/MainPanel/SConsoleVariablesEditorMainPanel.cpp
+++ b/Engine/Plugins/Experimental/ConsoleVariables/Source/ConsoleVariablesEditor/Private/Views/MainPanel/SConsoleVariablesEditorMainPanel.cpp
@@ -12,9 +12,12 @@
 #include "EditorFontGlyphs.h"
 #include "EditorStyleSet.h"
 #include "Framework/Application/SlateApplication.h"
+#include "Framework/MultiBox/MultiBoxBuilder.h"
 #include "IContentBrowserSingleton.h"
 #include "OutputLog/Public/OutputLogModule.h"
+#include "Widgets/Images/SImage.h"
 #include "Widgets/Input/SButton.h"
+#include "Widgets/Input/SComboButton.h"
 #include "Widgets/Input/SMultiLineEditableTextBox.h"
 #include "Widgets/Layout/SBox.h"
 #include "Widgets/Layout/SSplitter.h"
