Console Variables Editor:

-Multi-User support added
 -Changing variable values in dialogue will call SendConsoleVariableChange
 -If the user is tracking changed variables from outside the dialogue, this will also call SendConsoleVariableChange so all clients will begin tracking it as well
 -Added IsInitialized returning bool to FManager so that MU UI can be generated based on whether or not the Concert plugin is enabled
-Fix for: Some variables that don't trigger the sink delegate when changed outside the CVE dialogue
-Fixed issue where "FindSavedValueByCommandString" always returned false even when a value was found
-Removed redundant Settings functions from FConsoleVariablesEditorModule

diff --git a/Engine/Plugins/Experimental/ConsoleVariables/Source/ConsoleVariablesEditor/Private/ConsoleVariablesAsset.cpp b/Engine/Plugins/Experimental/ConsoleVariables/Source/ConsoleVariablesEditor/Private/ConsoleVariablesAsset.cpp
--- a/Engine/Plugins/Experimental/ConsoleVariables/Source/ConsoleVariablesEditor/Private/ConsoleVariablesAsset.cpp
+++ b/Engine/Plugins/Experimental/ConsoleVariables/Source/ConsoleVariablesEditor/Private/ConsoleVariablesAsset.cpp
@@ -24,6 +24,7 @@ bool UConsoleVariablesAsset::FindSavedValueByCommandString(const FString& InComm
 	if (const bool bMatchFound = KeyArray.Contains(InCommandString))
 	{
 		OutValue = SavedCommandsAndValues[InCommandString];
+		return true;
 	}
 	
 	return false;
diff --git a/Engine/Plugins/Experimental/ConsoleVariables/Source/ConsoleVariablesEditor/Private/ConsoleVariablesEditorCommandInfo.h b/Engine/Plugins/Experimental/ConsoleVariables/Source/ConsoleVariablesEditor/Private/ConsoleVariablesEditorCommandInfo.h
--- a/Engine/Plugins/Experimental/ConsoleVariables/Source/ConsoleVariablesEditor/Private/ConsoleVariablesEditorCommandInfo.h
+++ b/Engine/Plugins/Experimental/ConsoleVariables/Source/ConsoleVariablesEditor/Private/ConsoleVariablesEditorCommandInfo.h
@@ -21,25 +21,27 @@ struct FConsoleVariablesEditorCommandInfo
 		ConsoleVariablePtr = nullptr;
 	}
 
-	FConsoleVariablesEditorCommandInfo(const FString& InCommand, IConsoleVariable* InVariablePtr, const FString& InStartupValue)
+	FConsoleVariablesEditorCommandInfo(
+		const FString& InCommand, IConsoleVariable* InVariablePtr, const FString& InStartupValue, const FDelegateHandle& InOnVariableChangedCallbackHandle)
 	: Command(InCommand)
 	, ConsoleVariablePtr(InVariablePtr)
 	, StartupValueAsString(InStartupValue)
+	, OnVariableChangedCallbackHandle(InOnVariableChangedCallbackHandle)
 	{}
 
 	FORCEINLINE bool operator==(const FConsoleVariablesEditorCommandInfo& Comparator) const
 	{
 		return ConsoleVariablePtr && Comparator.ConsoleVariablePtr &&
-			ConsoleVariablePtr.Get() == Comparator.ConsoleVariablePtr.Get() &&
+			ConsoleVariablePtr == Comparator.ConsoleVariablePtr &&
 			Command.Equals(Comparator.Command);
 	}
 
 	void ExecuteCommand(const FString& NewValueAsString) const
 	{
-		GEngine->Exec(GEditor->GetWorld(), *FString::Printf(TEXT("%s %s"), *Command, *NewValueAsString));
+		GEngine->Exec(GetCurrentWorld(), *FString::Printf(TEXT("%s %s"), *Command, *NewValueAsString));
 	}
 	
-	UWorld* GetCurrentWorld() const
+	static UWorld* GetCurrentWorld()
 	{
 		UWorld* CurrentWorld = nullptr;
 		if (GIsEditor)
@@ -104,7 +106,7 @@ struct FConsoleVariablesEditorCommandInfo
 
 	bool IsCurrentValueDifferentFromInputValue(const FString& InValueToCompare) const
 	{
-		if (ConsoleVariablePtr && !ConsoleVariablePtr.IsNull())
+		if (ConsoleVariablePtr)
 		{
 			return !ConsoleVariablePtr->GetString().Equals(InValueToCompare);
 		}
@@ -116,10 +118,12 @@ struct FConsoleVariablesEditorCommandInfo
 	UPROPERTY()
 	FString Command;
 
-	TObjectPtr<IConsoleVariable> ConsoleVariablePtr;
+	IConsoleVariable* ConsoleVariablePtr;
 
 	/** The value of this command when the module started in this session after it may have been set by an ini file */
 	FString StartupValueAsString;
+	
+	FDelegateHandle OnVariableChangedCallbackHandle;
 };
 
 #undef LOCTEXT_NAMESPACE
diff --git a/Engine/Plugins/Experimental/ConsoleVariables/Source/ConsoleVariablesEditor/Private/ConsoleVariablesEditorModule.cpp b/Engine/Plugins/Experimental/ConsoleVariables/Source/ConsoleVariablesEditor/Private/ConsoleVariablesEditorModule.cpp
--- a/Engine/Plugins/Experimental/ConsoleVariables/Source/ConsoleVariablesEditor/Private/ConsoleVariablesEditorModule.cpp
+++ b/Engine/Plugins/Experimental/ConsoleVariables/Source/ConsoleVariablesEditor/Private/ConsoleVariablesEditorModule.cpp
@@ -4,13 +4,14 @@
 
 #include "AssetTypeActions/AssetTypeActions_ConsoleVariables.h"
 #include "ConsoleVariablesEditorLog.h"
+#include "ConsoleVariablesEditorProjectSettings.h"
 #include "ConsoleVariablesEditorStyle.h"
+#include "MultiUser/ConsoleVariableSyncData.h"
 #include "Views/MainPanel/ConsoleVariablesEditorMainPanel.h"
 
 #include "Algo/Find.h"
 #include "Framework/Docking/TabManager.h"
 #include "ISettingsModule.h"
-#include "ISettingsSection.h"
 #include "LevelEditor.h"
 #include "ToolMenus.h"
 #include "WorkspaceMenuStructure.h"
@@ -33,7 +34,6 @@ void FConsoleVariablesEditorModule::StartupModule()
 
 	FConsoleVariablesEditorStyle::Initialize();
 
-	FCoreDelegates::OnPostEngineInit.AddRaw(this, &FConsoleVariablesEditorModule::PostEngineInit);
 	FCoreDelegates::OnFEngineLoopInitComplete.AddRaw(this, &FConsoleVariablesEditorModule::OnFEngineLoopInitComplete);
 }
 
@@ -41,26 +41,28 @@ void FConsoleVariablesEditorModule::ShutdownModule()
 {
 	UToolMenus::UnregisterOwner(this);
 
-	FCoreDelegates::OnPostEngineInit.RemoveAll(this);
 	FCoreDelegates::OnFEngineLoopInitComplete.RemoveAll(this);
 
 	FConsoleVariablesEditorStyle::Shutdown();
+	
+	MainPanel.Reset();
+
+	ConsoleVariablesMasterReference.Empty();
 
 	// Unregister project settings
 	ISettingsModule& SettingsModule = FModuleManager::LoadModuleChecked<ISettingsModule>("Settings");
 	{
 		SettingsModule.UnregisterSettings("Project", "Plugins", "Console Variables Editor");
 	}
-	
-	MainPanel.Reset();
-	
-	ProjectSettingsSectionPtr.Reset();
-	ProjectSettingsObjectPtr.Reset();
-
-	ConsoleVariablesMasterReference.Empty();
 
-	IConsoleManager::Get().UnregisterConsoleVariableSink_Handle(VariableChangedSinkHandle);
-	VariableChangedSinkDelegate.Unbind();
+	// Remove all OnChanged delegates
+	for (TSharedPtr<FConsoleVariablesEditorCommandInfo> CommandInfo : ConsoleVariablesMasterReference)
+	{
+		if (CommandInfo.IsValid() && CommandInfo->ConsoleVariablePtr)
+		{
+			CommandInfo->ConsoleVariablePtr->OnChangedDelegate().Remove(CommandInfo->OnVariableChangedCallbackHandle);
+		}
+	}
 }
 
 void FConsoleVariablesEditorModule::OpenConsoleVariablesDialogWithAssetSelected(const FAssetData& InAssetData)
@@ -76,11 +78,6 @@ void FConsoleVariablesEditorModule::OpenConsoleVariablesDialogWithAssetSelected(
 	}
 }
 
-void FConsoleVariablesEditorModule::OpenConsoleVariablesSettings()
-{
-	FModuleManager::LoadModuleChecked<ISettingsModule>("Settings").ShowViewer("Project", "Plugins", "Console Variables Editor");
-}
-
 void FConsoleVariablesEditorModule::QueryAndBeginTrackingConsoleVariables()
 {
 	const int32 VariableCount = ConsoleVariablesMasterReference.Num();
@@ -92,8 +89,9 @@ void FConsoleVariablesEditorModule::QueryAndBeginTrackingConsoleVariables()
 		{
 			if (IConsoleVariable* AsVariable = ConsoleObject->AsVariable())
 			{
+				const FDelegateHandle Handle = AsVariable->OnChangedDelegate().AddRaw(this, &FConsoleVariablesEditorModule::OnConsoleVariableChanged);
 				ConsoleVariablesMasterReference.Add(
-					MakeShared<FConsoleVariablesEditorCommandInfo>(Key, AsVariable, AsVariable->GetString()));
+					MakeShared<FConsoleVariablesEditorCommandInfo>(Key, AsVariable, AsVariable->GetString(), Handle));
 			}
 		}),
 		TEXT(""));
@@ -111,21 +109,41 @@ TWeakPtr<FConsoleVariablesEditorCommandInfo> FConsoleVariablesEditorModule::Find
 	return Match ? *Match : nullptr;
 }
 
-void FConsoleVariablesEditorModule::PostEngineInit()
+TWeakPtr<FConsoleVariablesEditorCommandInfo> FConsoleVariablesEditorModule::FindCommandInfoByConsoleVariableReference(IConsoleVariable* InVariableReference)
 {
-	RegisterMenuItem();
-	RegisterProjectSettings();
+	TSharedPtr<FConsoleVariablesEditorCommandInfo>* Match = Algo::FindByPredicate(
+	ConsoleVariablesMasterReference,
+	[InVariableReference](const TSharedPtr<FConsoleVariablesEditorCommandInfo> Comparator)
+	{
+		return Comparator->ConsoleVariablePtr == InVariableReference;
+	});
+
+	return Match ? *Match : nullptr;
+}
+
+TObjectPtr<UConsoleVariablesAsset> FConsoleVariablesEditorModule::GetEditingAsset() const
+{
+	return EditingAsset.Get();
+}
+
+void FConsoleVariablesEditorModule::SetEditingAsset(const TObjectPtr<UConsoleVariablesAsset> InEditingAsset)
+{
+	EditingAsset = InEditingAsset;
+}
+
+void FConsoleVariablesEditorModule::SendMultiUserConsoleVariableChange(const FString& InVariableName, const FString& InValueAsString)
+{
+	MainPanel->GetMultiUserManager().SendConsoleVariableChange(InVariableName, InValueAsString);
 }
 
 void FConsoleVariablesEditorModule::OnFEngineLoopInitComplete()
 {
+	RegisterMenuItem();
+	RegisterProjectSettings();
 	QueryAndBeginTrackingConsoleVariables();
 	AllocateTransientPreset();
 	
 	MainPanel = MakeShared<FConsoleVariablesEditorMainPanel>();
-
-	VariableChangedSinkDelegate = FConsoleCommandDelegate::CreateRaw(this, &FConsoleVariablesEditorModule::OnConsoleVariableChange);
-	VariableChangedSinkHandle = IConsoleManager::Get().RegisterConsoleVariableSink_Handle(VariableChangedSinkDelegate);
 }
 
 void FConsoleVariablesEditorModule::RegisterMenuItem()
@@ -146,25 +164,42 @@ bool FConsoleVariablesEditorModule::RegisterProjectSettings()
 	ISettingsModule& SettingsModule = FModuleManager::LoadModuleChecked<ISettingsModule>("Settings");
 	{
 		// User Project Settings
-		ProjectSettingsSectionPtr = SettingsModule.RegisterSettings("Project", "Plugins", "Console Variables Editor",
+		const TSharedPtr<ISettingsSection> ProjectSettingsSectionPtr = SettingsModule.RegisterSettings(
+			"Project", "Plugins", "Console Variables Editor",
 			NSLOCTEXT("ConsoleVariables", "ConsoleVariablesSettingsCategoryDisplayName", "Console Variables Editor"),
-			NSLOCTEXT("ConsoleVariables", "ConsoleVariablesSettingsDescription", "Configure the Console Variables Editor user settings"),
+			NSLOCTEXT("ConsoleVariables", "ConsoleVariablesSettingsDescription",
+			          "Configure the Console Variables Editor user settings"),
 			GetMutableDefault<UConsoleVariablesEditorProjectSettings>());
 
-		if (ProjectSettingsSectionPtr.IsValid() && ProjectSettingsSectionPtr->GetSettingsObject().IsValid())
-		{
-			ProjectSettingsObjectPtr = Cast<UConsoleVariablesEditorProjectSettings>(ProjectSettingsSectionPtr->GetSettingsObject());
-
-			ProjectSettingsSectionPtr->OnModified().BindRaw(this, &FConsoleVariablesEditorModule::HandleModifiedProjectSettings);
-		}
+		return true;
 	}
 
-	return ProjectSettingsObjectPtr.IsValid();
+	return false;
 }
 
-bool FConsoleVariablesEditorModule::HandleModifiedProjectSettings()
-{	
-	return true;
+void FConsoleVariablesEditorModule::OnConsoleVariableChanged(IConsoleVariable* ChangedVariable)
+{
+	check(EditingAsset);
+
+	if (const TWeakPtr<FConsoleVariablesEditorCommandInfo> CommandInfo =
+		FindCommandInfoByConsoleVariableReference(ChangedVariable); CommandInfo.IsValid())
+	{
+		FString OutValue;
+		const FString& Key = CommandInfo.Pin()->Command;
+		if (GetMutableDefault<UConsoleVariablesEditorProjectSettings>()->bAddAllChangedConsoleVariablesToCurrentPreset &&
+			!EditingAsset->FindSavedValueByCommandString(Key, OutValue) &&
+			CommandInfo.Pin()->IsCurrentValueDifferentFromInputValue(CommandInfo.Pin()->StartupValueAsString))
+		{
+			EditingAsset->AddOrSetConsoleVariableSavedValue(Key, ChangedVariable->GetString());
+
+			if (MainPanel.IsValid())
+			{
+				MainPanel->RefreshList();
+			}
+
+			SendMultiUserConsoleVariableChange(Key, ChangedVariable->GetString());
+		}
+	}
 }
 
 TObjectPtr<UConsoleVariablesAsset> FConsoleVariablesEditorModule::AllocateTransientPreset()
@@ -196,61 +231,6 @@ void FConsoleVariablesEditorModule::OpenConsoleVariablesEditor()
 	FGlobalTabmanager::Get()->TryInvokeTab(ConsoleVariablesToolkitPanelTabId);
 }
 
-void FConsoleVariablesEditorModule::OnConsoleVariableChange()
-{
-	check(EditingAsset);
-
-	const int32 StartingTrackedCommandsCount = EditingAsset->GetSavedCommandsAndValues().Num();
-
-	IConsoleManager::Get().ForEachConsoleObjectThatStartsWith(FConsoleObjectVisitor::CreateLambda(
-	[this] (const TCHAR* Key, IConsoleObject* ConsoleObject)
-	{
-		if (IConsoleVariable* AsVariable = ConsoleObject->AsVariable())
-		{
-			TWeakPtr<FConsoleVariablesEditorCommandInfo> CommandInfo = FindCommandInfoByName(Key);
-			
-			if (CommandInfo.IsValid())
-			{
-				CommandInfo.Pin()->ConsoleVariablePtr = AsVariable;
-						
-				if (ProjectSettingsObjectPtr->bAddAllChangedConsoleVariablesToCurrentPreset)
-				{
-					FString OutValue;
-					if (!EditingAsset->FindSavedValueByCommandString(Key, OutValue) &&
-						CommandInfo.Pin()->IsCurrentValueDifferentFromInputValue(CommandInfo.Pin()->StartupValueAsString))
-					{
-						EditingAsset->AddOrSetConsoleVariableSavedValue(Key, AsVariable->GetString());
-					}
-				}
-			}
-			else
-			{
-				ConsoleVariablesMasterReference.Add(
-					MakeShared<FConsoleVariablesEditorCommandInfo>(Key, AsVariable, AsVariable->GetString()));
-			}
-		}
-	}),
-	TEXT(""));
-
-	if (StartingTrackedCommandsCount < EditingAsset->GetSavedCommandsAndValues().Num())
-	{
-		if (MainPanel.IsValid())
-		{
-			MainPanel->RefreshList();
-		}
-	}
-}
-
-TObjectPtr<UConsoleVariablesAsset> FConsoleVariablesEditorModule::GetEditingAsset() const
-{
-	return EditingAsset.Get();
-}
-
-void FConsoleVariablesEditorModule::SetEditingAsset(const TObjectPtr<UConsoleVariablesAsset> InEditingAsset)
-{
-	EditingAsset = InEditingAsset;
-}
-
 #undef LOCTEXT_NAMESPACE
 	
 IMPLEMENT_MODULE(FConsoleVariablesEditorModule, ConsoleVariablesEditor)
diff --git a/Engine/Plugins/Experimental/ConsoleVariables/Source/ConsoleVariablesEditor/Private/ConsoleVariablesEditorModule.h b/Engine/Plugins/Experimental/ConsoleVariables/Source/ConsoleVariablesEditor/Private/ConsoleVariablesEditorModule.h
--- a/Engine/Plugins/Experimental/ConsoleVariables/Source/ConsoleVariablesEditor/Private/ConsoleVariablesEditorModule.h
+++ b/Engine/Plugins/Experimental/ConsoleVariables/Source/ConsoleVariablesEditor/Private/ConsoleVariablesEditorModule.h
@@ -4,7 +4,7 @@
 
 #include "ConsoleVariablesAsset.h"
 #include "ConsoleVariablesEditorCommandInfo.h"
-#include "ConsoleVariablesEditorProjectSettings.h"
+#include "MultiUser/ConsoleVariableSync.h"
 
 #include "CoreMinimal.h"
 #include "HAL/IConsoleManager.h"
@@ -29,39 +29,35 @@ public:
 	//~ End IModuleInterface Interface
 
 	void OpenConsoleVariablesDialogWithAssetSelected(const FAssetData& InAssetData);
-	
-	static void OpenConsoleVariablesSettings();
-
-	TWeakObjectPtr<UConsoleVariablesEditorProjectSettings> GetConsoleVariablesUserSettings() const
-	{
-		return ProjectSettingsObjectPtr;
-	}
 
 	/** Find all console variables and cache their startup values */
 	void QueryAndBeginTrackingConsoleVariables();
 
 	/** Find a tracked console variable by the command string with optional case sensitivity. */
 	TWeakPtr<FConsoleVariablesEditorCommandInfo> FindCommandInfoByName(const FString& NameToSearch, ESearchCase::Type InSearchCase = ESearchCase::IgnoreCase);
+
+	/** Find a tracked console variable by its variable reference. */
+	TWeakPtr<FConsoleVariablesEditorCommandInfo> FindCommandInfoByConsoleVariableReference(IConsoleVariable* InVariableReference);
 	
 	[[nodiscard]] TObjectPtr<UConsoleVariablesAsset> GetEditingAsset() const;
 	void SetEditingAsset(const TObjectPtr<UConsoleVariablesAsset> InEditingAsset);
+
+	void SendMultiUserConsoleVariableChange(const FString& InVariableName, const FString& InValueAsString);
 	
 private:
 
-	void PostEngineInit();
 	void OnFEngineLoopInitComplete();
 
 	void RegisterMenuItem();
 	bool RegisterProjectSettings();
-	bool HandleModifiedProjectSettings();
+
+	void OnConsoleVariableChanged(IConsoleVariable* ChangedVariable);
 
 	TObjectPtr<UConsoleVariablesAsset> AllocateTransientPreset();
 
 	TSharedRef<SDockTab> SpawnMainPanelTab(const FSpawnTabArgs& Args);
 
 	void OpenConsoleVariablesEditor();
-
-	void OnConsoleVariableChange();
 	
 	static const FName ConsoleVariablesToolkitPanelTabId;
 
@@ -71,14 +67,6 @@ private:
 	// Transient preset that's being edited so we don't affect the reference asset unless we save it
 	TObjectPtr<UConsoleVariablesAsset> EditingAsset = nullptr;
 
-	TSharedPtr<ISettingsSection> ProjectSettingsSectionPtr;
-	TWeakObjectPtr<UConsoleVariablesEditorProjectSettings> ProjectSettingsObjectPtr;
-
 	/** All tracked variables and their default, startup, and current values */
 	TArray<TSharedPtr<FConsoleVariablesEditorCommandInfo>> ConsoleVariablesMasterReference;
-
-	/** A callback registered with the Console Manager that is called when a console variable is changed */
-	FConsoleVariableSinkHandle VariableChangedSinkHandle;
-
-	FConsoleCommandDelegate VariableChangedSinkDelegate;
 };
diff --git a/Engine/Plugins/Experimental/ConsoleVariables/Source/ConsoleVariablesEditor/Private/ConsoleVariablesEditorStyle.cpp b/Engine/Plugins/Experimental/ConsoleVariables/Source/ConsoleVariablesEditor/Private/ConsoleVariablesEditorStyle.cpp
--- a/Engine/Plugins/Experimental/ConsoleVariables/Source/ConsoleVariablesEditor/Private/ConsoleVariablesEditorStyle.cpp
+++ b/Engine/Plugins/Experimental/ConsoleVariables/Source/ConsoleVariablesEditor/Private/ConsoleVariablesEditorStyle.cpp
@@ -36,6 +36,7 @@ const FSlateBrush* FConsoleVariablesEditorStyle::GetBrush(const FName PropertyNa
 #define IMAGE_BRUSH( RelativePath, ... ) FSlateImageBrush( Style->RootToContentDir( RelativePath, TEXT(".png") ), __VA_ARGS__ )
 #define BOX_BRUSH( RelativePath, ... ) FSlateBoxBrush( Style->RootToContentDir( RelativePath, TEXT(".png") ), __VA_ARGS__ )
 #define BORDER_BRUSH( RelativePath, ... ) FSlateBorderBrush( Style->RootToContentDir( RelativePath, TEXT(".png") ), __VA_ARGS__ )
+#define IMAGE_PLUGIN_BRUSH_SVG( RelativePath, ... ) FSlateVectorImageBrush( FConsoleVariablesEditorStyle::InContent(RelativePath, ".svg"), __VA_ARGS__)
 
 const FVector2D Icon64x64(64.f, 64.f);
 const FVector2D Icon40x40(40.0f, 40.0f);
@@ -44,6 +45,12 @@ const FVector2D Icon16x16(16.0f, 16.0f);
 const FVector2D Icon12x12(12.0f, 12.0f);
 const FVector2D Icon8x8(8.f, 8.f);
 
+FString FConsoleVariablesEditorStyle::InContent(const FString& RelativePath, const ANSICHAR* Extension)
+{
+	static FString ContentDir = IPluginManager::Get().FindPlugin(TEXT("ConcertSyncClient"))->GetContentDir();
+	return (ContentDir / RelativePath) + Extension;
+}
+
 TSharedRef< FSlateStyleSet > FConsoleVariablesEditorStyle::Create()
 {
 	TSharedRef<FSlateStyleSet> Style = MakeShared<FSlateStyleSet>("ConsoleVariablesEditor");
@@ -69,6 +76,9 @@ TSharedRef< FSlateStyleSet > FConsoleVariablesEditorStyle::Create()
 	Style->Set("ConsoleVariablesEditor.HeaderRowBorder", new FSlateColorBrush(FStyleColors::Black));
 	Style->Set("ConsoleVariablesEditor.CommandGroupBorder", new BOX_BRUSH("Common/DarkGroupBorder", FMargin(4.0f / 16.0f)));
 	Style->Set("ConsoleVariablesEditor.DefaultBorder", new FSlateColorBrush(FStyleColors::Transparent));
+
+	// Multi-user Tab/Menu icons
+	Style->Set("Concert.MultiUser", new IMAGE_PLUGIN_BRUSH_SVG("Icons/icon_MultiUser", Icon16x16));
 	
 	return Style;
 }
diff --git a/Engine/Plugins/Experimental/ConsoleVariables/Source/ConsoleVariablesEditor/Private/ConsoleVariablesEditorStyle.h b/Engine/Plugins/Experimental/ConsoleVariables/Source/ConsoleVariablesEditor/Private/ConsoleVariablesEditorStyle.h
--- a/Engine/Plugins/Experimental/ConsoleVariables/Source/ConsoleVariablesEditor/Private/ConsoleVariablesEditorStyle.h
+++ b/Engine/Plugins/Experimental/ConsoleVariables/Source/ConsoleVariablesEditor/Private/ConsoleVariablesEditorStyle.h
@@ -27,6 +27,9 @@ public:
 	}
 
 private:
+	
+	static FString InContent(const FString& RelativePath, const ANSICHAR* Extension);
+	
 	static TSharedRef<FSlateStyleSet> Create();
 
 	static TSharedPtr<FSlateStyleSet> StyleInstance;
diff --git a/Engine/Plugins/Experimental/ConsoleVariables/Source/ConsoleVariablesEditor/Private/MultiUser/ConcertConsoleVariableSessionCustomization.cpp b/Engine/Plugins/Experimental/ConsoleVariables/Source/ConsoleVariablesEditor/Private/MultiUser/ConcertConsoleVariableSessionCustomization.cpp
--- a/Engine/Plugins/Experimental/ConsoleVariables/Source/ConsoleVariablesEditor/Private/MultiUser/ConcertConsoleVariableSessionCustomization.cpp
+++ b/Engine/Plugins/Experimental/ConsoleVariables/Source/ConsoleVariablesEditor/Private/MultiUser/ConcertConsoleVariableSessionCustomization.cpp
@@ -402,7 +402,7 @@ void FConcertConsoleVariableSessionCustomization::CustomizeDetails(IDetailLayout
 				+SHeaderRow::Column(ConsoleVariableSessionUI::DisplayNameColumnName)
 				.DefaultLabel(LOCTEXT("UserDisplayName", "Display Name"))
 				+SHeaderRow::Column(ConsoleVariableSessionUI::TransactColumnName)
-				.DefaultLabel(LOCTEXT("Transact", "Transact"))
+				.DefaultLabel(LOCTEXT("Transact", "Transact Console Variables"))
 			)
 		];
 	}
diff --git a/Engine/Plugins/Experimental/ConsoleVariables/Source/ConsoleVariablesEditor/Private/MultiUser/ConsoleVariableSync.cpp b/Engine/Plugins/Experimental/ConsoleVariables/Source/ConsoleVariablesEditor/Private/MultiUser/ConsoleVariableSync.cpp
--- a/Engine/Plugins/Experimental/ConsoleVariables/Source/ConsoleVariablesEditor/Private/MultiUser/ConsoleVariableSync.cpp
+++ b/Engine/Plugins/Experimental/ConsoleVariables/Source/ConsoleVariablesEditor/Private/MultiUser/ConsoleVariableSync.cpp
@@ -41,6 +41,8 @@ struct FManagerImpl
 			ConcertClient->OnSessionConnectionChanged().AddRaw(this, &FManagerImpl::OnSessionConnectionChanged);
 
 			RegisterExtensions();
+
+			bIsInitialized = true;
 		}
 	}
 
@@ -423,6 +425,9 @@ struct FManagerImpl
 
 	/** Indicates if the multi-user console synchronization is enabled. */
 	bool bIsEnabled = true;
+
+	/** True if the Concert module was found and customizations have been registered. */
+	bool bIsInitialized = false;
 };
 
 
@@ -453,5 +458,9 @@ void FManager::SetEnableMultiUserSupport(bool bIsEnabled)
 	Implementation->bIsEnabled = bIsEnabled;
 }
 
+bool FManager::IsInitialized() const
+{
+	return Implementation->bIsInitialized;
+}
 };
 
diff --git a/Engine/Plugins/Experimental/ConsoleVariables/Source/ConsoleVariablesEditor/Private/MultiUser/ConsoleVariableSync.h b/Engine/Plugins/Experimental/ConsoleVariables/Source/ConsoleVariablesEditor/Private/MultiUser/ConsoleVariableSync.h
--- a/Engine/Plugins/Experimental/ConsoleVariables/Source/ConsoleVariablesEditor/Private/MultiUser/ConsoleVariableSync.h
+++ b/Engine/Plugins/Experimental/ConsoleVariables/Source/ConsoleVariablesEditor/Private/MultiUser/ConsoleVariableSync.h
@@ -39,6 +39,9 @@ struct FManager
 
 	/** Enables / disables multi-user message handling. */
 	void SetEnableMultiUserSupport(bool bIsEnabled);
+
+	bool IsInitialized() const;
+	
 private:
 	TUniquePtr<FManagerImpl> Implementation;
 };
diff --git a/Engine/Plugins/Experimental/ConsoleVariables/Source/ConsoleVariablesEditor/Private/Views/List/SConsoleVariablesEditorListRow.cpp b/Engine/Plugins/Experimental/ConsoleVariables/Source/ConsoleVariablesEditor/Private/Views/List/SConsoleVariablesEditorListRow.cpp
--- a/Engine/Plugins/Experimental/ConsoleVariables/Source/ConsoleVariablesEditor/Private/Views/List/SConsoleVariablesEditorListRow.cpp
+++ b/Engine/Plugins/Experimental/ConsoleVariables/Source/ConsoleVariablesEditor/Private/Views/List/SConsoleVariablesEditorListRow.cpp
@@ -2,6 +2,7 @@
 
 #include "SConsoleVariablesEditorListRow.h"
 
+#include "ConsoleVariablesEditorModule.h"
 #include "ConsoleVariablesEditorStyle.h"
 #include "SConsoleVariablesEditorListValueInput.h"
 
@@ -241,10 +242,16 @@ void SConsoleVariablesEditorListRow::OnCheckboxStateChange(const ECheckBoxState
 			if (PinnedItem->IsRowChecked())
 			{
 				PinnedItem->GetCommandInfo().Pin()->ExecuteCommand(ValueChildInputWidget->GetCachedValue());
+
+				FConsoleVariablesEditorModule::Get().SendMultiUserConsoleVariableChange(
+					PinnedItem->GetCommandInfo().Pin()->Command, ValueChildInputWidget->GetCachedValue());
 			}
 			else
 			{
 				PinnedItem->GetCommandInfo().Pin()->ExecuteCommand(PinnedItem->GetCommandInfo().Pin()->StartupValueAsString);
+
+				FConsoleVariablesEditorModule::Get().SendMultiUserConsoleVariableChange(
+					PinnedItem->GetCommandInfo().Pin()->Command, PinnedItem->GetCommandInfo().Pin()->StartupValueAsString);
 			}
 		}
 
diff --git a/Engine/Plugins/Experimental/ConsoleVariables/Source/ConsoleVariablesEditor/Private/Views/List/SConsoleVariablesEditorListValueInput.cpp b/Engine/Plugins/Experimental/ConsoleVariables/Source/ConsoleVariablesEditor/Private/Views/List/SConsoleVariablesEditorListValueInput.cpp
--- a/Engine/Plugins/Experimental/ConsoleVariables/Source/ConsoleVariablesEditor/Private/Views/List/SConsoleVariablesEditorListValueInput.cpp
+++ b/Engine/Plugins/Experimental/ConsoleVariables/Source/ConsoleVariablesEditor/Private/Views/List/SConsoleVariablesEditorListValueInput.cpp
@@ -3,6 +3,7 @@
 #include "SConsoleVariablesEditorListValueInput.h"
 
 #include "ConsoleVariablesEditorListRow.h"
+#include "ConsoleVariablesEditorModule.h"
 
 #include "Widgets/Input/SNumericEntryBox.h"
 
@@ -76,9 +77,13 @@ void SConsoleVariablesEditorListValueInput_Float::Construct(const FArguments& In
 		})
 		.OnValueChanged_Lambda([this] (const float InValue)
 		{
-			if (Item.IsValid() && Item.Pin()->GetWidgetCheckedState() == ECheckBoxState::Checked)
+			if (const TSharedPtr<FConsoleVariablesEditorListRow> PinnedItem = Item.Pin(); PinnedItem->GetWidgetCheckedState() == ECheckBoxState::Checked)
 			{
-				Item.Pin()->GetCommandInfo().Pin()->ExecuteCommand(FString::SanitizeFloat(InValue));
+				const FString ValueAsString = FString::SanitizeFloat(InValue);
+				
+				PinnedItem->GetCommandInfo().Pin()->ExecuteCommand(ValueAsString);
+
+				FConsoleVariablesEditorModule::Get().SendMultiUserConsoleVariableChange(PinnedItem->GetCommandInfo().Pin()->Command, ValueAsString);
 			}
 			
 			SetCachedValue(FString::SanitizeFloat(InValue));
@@ -122,9 +127,13 @@ void SConsoleVariablesEditorListValueInput_Int::Construct(const FArguments& InAr
 		})
 		.OnValueChanged_Lambda([this] (const int32 InValue)
 		{
-			if (Item.IsValid() && Item.Pin()->GetWidgetCheckedState() == ECheckBoxState::Checked)
+			if (const TSharedPtr<FConsoleVariablesEditorListRow> PinnedItem = Item.Pin(); PinnedItem->GetWidgetCheckedState() == ECheckBoxState::Checked)
 			{
-				Item.Pin()->GetCommandInfo().Pin()->ExecuteCommand(FString::FromInt(InValue));
+				const FString ValueAsString = FString::FromInt(InValue);
+				
+				PinnedItem->GetCommandInfo().Pin()->ExecuteCommand(ValueAsString);
+
+				FConsoleVariablesEditorModule::Get().SendMultiUserConsoleVariableChange(PinnedItem->GetCommandInfo().Pin()->Command, ValueAsString);
 			}
 
 			SetCachedValue(FString::FromInt(InValue));
@@ -168,9 +177,13 @@ void SConsoleVariablesEditorListValueInput_String::Construct(const FArguments& I
 		})
 		.OnTextCommitted_Lambda([this] (const FText& InValue, ETextCommit::Type InTextCommitType)
 		{
-			if (InTextCommitType == ETextCommit::OnEnter && Item.IsValid() && Item.Pin()->GetWidgetCheckedState() == ECheckBoxState::Checked)
+			if (const TSharedPtr<FConsoleVariablesEditorListRow> PinnedItem = Item.Pin(); PinnedItem->GetWidgetCheckedState() == ECheckBoxState::Checked)
 			{
-				Item.Pin()->GetCommandInfo().Pin()->ExecuteCommand(InValue.ToString());
+				const FString ValueAsString = InValue.ToString();
+				
+				PinnedItem->GetCommandInfo().Pin()->ExecuteCommand(ValueAsString);
+
+				FConsoleVariablesEditorModule::Get().SendMultiUserConsoleVariableChange(PinnedItem->GetCommandInfo().Pin()->Command, ValueAsString);
 			}
 
 			SetCachedValue(InValue.ToString());
@@ -216,9 +229,13 @@ void SConsoleVariablesEditorListValueInput_Bool::Construct(const FArguments& InA
 		.MaxSliderValue(2)
 		.OnValueChanged_Lambda([this] (const int32 InValue)
 		{
-			if (Item.IsValid() && Item.Pin()->GetWidgetCheckedState() == ECheckBoxState::Checked)
+			if (const TSharedPtr<FConsoleVariablesEditorListRow> PinnedItem = Item.Pin(); PinnedItem->GetWidgetCheckedState() == ECheckBoxState::Checked)
 			{
-				Item.Pin()->GetCommandInfo().Pin()->ExecuteCommand(FString::FromInt(InValue));
+				const FString ValueAsString = FString::FromInt(InValue);
+				
+				PinnedItem->GetCommandInfo().Pin()->ExecuteCommand(ValueAsString);
+
+				FConsoleVariablesEditorModule::Get().SendMultiUserConsoleVariableChange(PinnedItem->GetCommandInfo().Pin()->Command, ValueAsString);
 			}
 
 			SetCachedValue(FString::FromInt(InValue));
diff --git a/Engine/Plugins/Experimental/ConsoleVariables/Source/ConsoleVariablesEditor/Private/Views/MainPanel/ConsoleVariablesEditorMainPanel.cpp b/Engine/Plugins/Experimental/ConsoleVariables/Source/ConsoleVariablesEditor/Private/Views/MainPanel/ConsoleVariablesEditorMainPanel.cpp
--- a/Engine/Plugins/Experimental/ConsoleVariables/Source/ConsoleVariablesEditor/Private/Views/MainPanel/ConsoleVariablesEditorMainPanel.cpp
+++ b/Engine/Plugins/Experimental/ConsoleVariables/Source/ConsoleVariablesEditor/Private/Views/MainPanel/ConsoleVariablesEditorMainPanel.cpp
@@ -11,31 +11,23 @@
 
 #include "FileHelpers.h"
 #include "Framework/Application/SlateApplication.h"
+#include "MultiUser/ConsoleVariableSyncData.h"
 
 
 FConsoleVariablesEditorMainPanel::FConsoleVariablesEditorMainPanel()
 {
 	EditorList = MakeShared<FConsoleVariablesEditorList>();
+	
+	OnConnectionChangedHandle = MultiUserManager.OnConnectionChange().AddRaw(
+		this, &FConsoleVariablesEditorMainPanel::OnConnectionChanged);
+	OnRemoteCVarChangeHandle = MultiUserManager.OnRemoteCVarChange().AddRaw(
+		this, &FConsoleVariablesEditorMainPanel::OnRemoteCvarChange);
+}
 
-	MultiUserManager.OnConnectionChange().AddLambda([] (EConcertConnectionStatus Status)
-		{
-			switch(Status) {
-			case EConcertConnectionStatus::Connected:
-				UE_LOG(LogConsoleVariablesEditor, Display, TEXT("Multi-user has connected to a session."));
-				break;
-			case EConcertConnectionStatus::Disconnected:
-				UE_LOG(LogConsoleVariablesEditor, Display, TEXT("Multi-user has disconnect from session."));
-				break;
-			default:
-				break;
-			};
-
-		});
-
-	MultiUserManager.OnRemoteCVarChange().AddLambda([] (FString InName, FString InValue)
-		{
-			UE_LOG(LogConsoleVariablesEditor, Display, TEXT("Remote set console variable %s = %s"), *InName, *InValue);
-		});
+FConsoleVariablesEditorMainPanel::~FConsoleVariablesEditorMainPanel()
+{
+	MultiUserManager.OnConnectionChange().Remove(OnConnectionChangedHandle);
+	MultiUserManager.OnRemoteCVarChange().Remove(OnRemoteCVarChangeHandle);
 }
 
 TSharedRef<SWidget> FConsoleVariablesEditorMainPanel::GetOrCreateWidget()
@@ -89,6 +81,14 @@ void FConsoleVariablesEditorMainPanel::UpdatePresetValuesForSave(TObjectPtr<UCon
 	}
 }
 
+void FConsoleVariablesEditorMainPanel::RefreshMultiUserDetails() const
+{
+	if (MainPanelWidget.IsValid())
+	{
+		MainPanelWidget->RefreshMultiUserDetails();
+	}
+}
+
 void FConsoleVariablesEditorMainPanel::SavePreset()
 {
 	const TWeakObjectPtr<UConsoleVariablesAsset> EditingAsset = GetEditingAsset();
@@ -162,3 +162,29 @@ bool FConsoleVariablesEditorMainPanel::ImportPreset_Impl(const FAssetData& InPre
 
 	return false;
 }
+
+void FConsoleVariablesEditorMainPanel::OnConnectionChanged(EConcertConnectionStatus Status)
+{
+	switch (Status)
+	{
+		case EConcertConnectionStatus::Connected:
+			UE_LOG(LogConsoleVariablesEditor, Display, TEXT("Multi-user has connected to a session."));
+			break;
+		case EConcertConnectionStatus::Disconnected:
+			UE_LOG(LogConsoleVariablesEditor, Display, TEXT("Multi-user has disconnected from session."));
+			break;
+		default:
+			break;
+	}
+}
+
+void FConsoleVariablesEditorMainPanel::OnRemoteCvarChange(const FString InName, const FString InValue)
+{
+	UE_LOG(LogConsoleVariablesEditor, Display, TEXT("Remote set console variable %s = %s"), *InName, *InValue);
+
+	if (GetMutableDefault<UConcertCVarSynchronization>()->bSyncCVarTransactions)
+	{
+		GEngine->Exec(FConsoleVariablesEditorCommandInfo::GetCurrentWorld(),
+					  *FString::Printf(TEXT("%s %s"), *InName, *InValue));
+	}
+}
diff --git a/Engine/Plugins/Experimental/ConsoleVariables/Source/ConsoleVariablesEditor/Private/Views/MainPanel/ConsoleVariablesEditorMainPanel.h b/Engine/Plugins/Experimental/ConsoleVariables/Source/ConsoleVariablesEditor/Private/Views/MainPanel/ConsoleVariablesEditorMainPanel.h
--- a/Engine/Plugins/Experimental/ConsoleVariables/Source/ConsoleVariablesEditor/Private/Views/MainPanel/ConsoleVariablesEditorMainPanel.h
+++ b/Engine/Plugins/Experimental/ConsoleVariables/Source/ConsoleVariablesEditor/Private/Views/MainPanel/ConsoleVariablesEditorMainPanel.h
@@ -2,7 +2,6 @@
 
 #pragma once
 
-#include "MultiUser/ConsoleVariableSync.h"
 #include "ConsoleVariablesEditorModule.h"
 #include "Widgets/SWidget.h"
 
@@ -15,6 +14,8 @@ class FConsoleVariablesEditorMainPanel : public TSharedFromThis<FConsoleVariable
 public:
 	FConsoleVariablesEditorMainPanel();
 
+	~FConsoleVariablesEditorMainPanel();
+
 	TSharedRef<SWidget> GetOrCreateWidget();
 
 	static FConsoleVariablesEditorModule& GetConsoleVariablesModule();
@@ -25,6 +26,8 @@ public:
 	void RefreshList(const FString& InConsoleCommandToScrollTo = "") const;
 	void UpdatePresetValuesForSave(TObjectPtr<UConsoleVariablesAsset> InAsset) const;
 
+	void RefreshMultiUserDetails() const;
+
 	// Save / Load
 
 	void SavePreset();
@@ -41,6 +44,11 @@ public:
 		return EditorList;
 	}
 
+	UE::ConsoleVariables::MultiUser::Private::FManager& GetMultiUserManager()
+	{
+		return MultiUserManager;
+	}
+
 private:
 
 	bool ImportPreset_Impl(const FAssetData& InPresetAsset, const TObjectPtr<UConsoleVariablesAsset> EditingAsset);
@@ -52,5 +60,10 @@ private:
 
 	TSharedPtr<FConsoleVariablesEditorList> EditorList;
 
+	void OnConnectionChanged(EConcertConnectionStatus Status);
+	void OnRemoteCvarChange(const FString InName, const FString InValue);
+
 	UE::ConsoleVariables::MultiUser::Private::FManager MultiUserManager;
+	FDelegateHandle OnConnectionChangedHandle;
+	FDelegateHandle OnRemoteCVarChangeHandle;
 };
diff --git a/Engine/Plugins/Experimental/ConsoleVariables/Source/ConsoleVariablesEditor/Private/Views/MainPanel/SConsoleVariablesEditorMainPanel.cpp b/Engine/Plugins/Experimental/ConsoleVariables/Source/ConsoleVariablesEditor/Private/Views/MainPanel/SConsoleVariablesEditorMainPanel.cpp
--- a/Engine/Plugins/Experimental/ConsoleVariables/Source/ConsoleVariablesEditor/Private/Views/MainPanel/SConsoleVariablesEditorMainPanel.cpp
+++ b/Engine/Plugins/Experimental/ConsoleVariables/Source/ConsoleVariablesEditor/Private/Views/MainPanel/SConsoleVariablesEditorMainPanel.cpp
@@ -6,6 +6,7 @@
 #include "ConsoleVariablesEditorLog.h"
 #include "ConsoleVariablesEditorModule.h"
 #include "ConsoleVariablesEditorStyle.h"
+#include "MultiUser/ConsoleVariableSyncData.h"
 #include "Views/List/ConsoleVariablesEditorList.h"
 #include "Views/MainPanel/ConsoleVariablesEditorMainPanel.h"
 
@@ -40,7 +41,7 @@ void SConsoleVariablesEditorMainPanel::Construct(const FArguments& InArgs, const
 	check(ConsoleInput.IsValid());
 
 	ConsoleInputEditableTextBox->SetOnKeyDownHandler(FOnKeyDown::CreateRaw(this, &SConsoleVariablesEditorMainPanel::HandleConsoleInputTextCommitted));
-
+	
 	ChildSlot
 	[
 		SNew(SSplitter)
@@ -56,7 +57,15 @@ void SConsoleVariablesEditorMainPanel::Construct(const FArguments& InArgs, const
 		[
 			MainPanel.Pin()->GetEditorList().Pin()->GetOrCreateWidget()
 		]
+
+		+SSplitter::Slot()
+		[
+			SAssignNew(MultiUserDetailsBox, SVerticalBox)
+			.Visibility(EVisibility::Collapsed)
+		]
 	];
+
+	CreateConcertButtonIfNeeded();
 }
 
 SConsoleVariablesEditorMainPanel::~SConsoleVariablesEditorMainPanel()
@@ -95,6 +104,37 @@ FReply SConsoleVariablesEditorMainPanel::HandleConsoleInputTextCommitted(const F
 	return FReply::Handled();
 }
 
+void SConsoleVariablesEditorMainPanel::RefreshMultiUserDetails()
+{
+	UConcertCVarSynchronization* CVarSync = GetMutableDefault<UConcertCVarSynchronization>();
+
+	UConcertCVarConfig* CVarConfig = GetMutableDefault<UConcertCVarConfig>();
+
+	MultiUserDetailsBox->ClearChildren();
+
+	MultiUserDetailsBox->AddSlot()
+	.AutoHeight()
+	[
+		GetConcertDetailsWidget(CVarSync)
+	];
+
+	MultiUserDetailsBox->AddSlot()
+	[
+		GetConcertDetailsWidget(CVarConfig)
+	];
+}
+
+void SConsoleVariablesEditorMainPanel::ToggleMultiUserDetails(ECheckBoxState CheckState)
+{
+	const bool bShouldBeVisible = CheckState == ECheckBoxState::Checked;
+	MultiUserDetailsBox->SetVisibility(bShouldBeVisible ? EVisibility::SelfHitTestInvisible : EVisibility::Collapsed);
+
+	if (bShouldBeVisible)
+	{
+		RefreshMultiUserDetails();
+	}
+}
+
 TSharedRef<SWidget> SConsoleVariablesEditorMainPanel::GeneratePanelToolbar(const TSharedRef<SWidget> InConsoleInputWidget)
 {
 	return SNew(SBorder)
@@ -102,7 +142,7 @@ TSharedRef<SWidget> SConsoleVariablesEditorMainPanel::GeneratePanelToolbar(const
 	        .BorderImage(FAppStyle::Get().GetBrush("NoBorder"))
 			.HAlign(HAlign_Fill)
 	        [
-				SNew(SHorizontalBox)
+				SAssignNew(ToolbarHBox, SHorizontalBox)
 				
 				// Add Console Variable input
 				+ SHorizontalBox::Slot()
@@ -115,8 +155,7 @@ TSharedRef<SWidget> SConsoleVariablesEditorMainPanel::GeneratePanelToolbar(const
 
 				// Presets Management Button
 				+ SHorizontalBox::Slot()
-				.Padding(FMargin(0.f, 0.f, 4.f, 0.f))
-				.HAlign(HAlign_Center)
+				.HAlign(HAlign_Right)
 				.VAlign(VAlign_Fill)
 				.AutoWidth()
 				[
@@ -147,38 +186,40 @@ TSharedRef<SWidget> SConsoleVariablesEditorMainPanel::GeneratePanelToolbar(const
 						]
 					]
 				]
-
-				// Open Settings
-				+ SHorizontalBox::Slot()
-                .HAlign(HAlign_Right)
-				.VAlign(VAlign_Fill)
-                [
-					SNew(SBox)
-					.WidthOverride(28)
-					.HeightOverride(28)
-					.Visibility(EVisibility::Collapsed)
-					[
-						SAssignNew(SettingsButtonPtr, SCheckBox)
-						.Padding(FMargin(4.f))
-						.ToolTipText(LOCTEXT("ShowSettings_Tip", "Show the general user/project settings for Console Variables"))
-						.Style(&FAppStyle::Get().GetWidgetStyle<FCheckBoxStyle>("ToggleButtonCheckbox"))
-						.ForegroundColor(FStyleColors::Foreground)
-						.IsChecked(false)
-						.OnCheckStateChanged_Lambda([this](ECheckBoxState CheckState)
-						{
-							FConsoleVariablesEditorModule::OpenConsoleVariablesSettings();
-							SettingsButtonPtr->SetIsChecked(false);
-						})
-		                [
-							SNew(SImage)
-							.Image(FAppStyle::Get().GetBrush("Icons.Settings"))
-							.ColorAndOpacity(FSlateColor::UseForeground())
-		                ]
-					]
-                ]
         	];
 }
 
+void SConsoleVariablesEditorMainPanel::CreateConcertButtonIfNeeded()
+{
+	if (MainPanel.Pin()->GetMultiUserManager().IsInitialized())
+	{
+		// Toggle Multi-User Details
+		ToolbarHBox->AddSlot()
+		.HAlign(HAlign_Right)
+		.VAlign(VAlign_Fill)
+		.AutoWidth()
+		[
+			SNew(SBox)
+			.WidthOverride(28)
+			.HeightOverride(28)
+			[
+				SAssignNew(ConcertButtonPtr, SCheckBox)
+				.Padding(FMargin(4.f))
+				.ToolTipText(LOCTEXT("ShowConcertSettings_Tip", "Show the multi-user controls for Console Variables"))
+				.Style(&FAppStyle::Get().GetWidgetStyle<FCheckBoxStyle>("ToggleButtonCheckbox"))
+				.ForegroundColor(FStyleColors::Foreground)
+				.IsChecked(false)
+				.OnCheckStateChanged_Raw(this, &SConsoleVariablesEditorMainPanel::ToggleMultiUserDetails)
+				[
+					SNew(SImage)
+					.Image(FConsoleVariablesEditorStyle::Get().GetBrush("Concert.MultiUser"))
+					.ColorAndOpacity(FSlateColor::UseForeground())
+				]
+			]
+		];
+	}
+}
+
 TSharedRef<SWidget> SConsoleVariablesEditorMainPanel::OnGeneratePresetsMenu()
 {
 	FMenuBuilder MenuBuilder(true, nullptr);
@@ -249,4 +290,21 @@ TSharedRef<SWidget> SConsoleVariablesEditorMainPanel::OnGeneratePresetsMenu()
 	return MenuBuilder.MakeWidget();
 }
 
+TSharedRef<SWidget> SConsoleVariablesEditorMainPanel::GetConcertDetailsWidget(UObject* InObject)
+{
+	FPropertyEditorModule& PropertyEditorModule =
+		FModuleManager::Get().LoadModuleChecked<FPropertyEditorModule>("PropertyEditor");
+	FDetailsViewArgs DetailsViewArgs;
+	DetailsViewArgs.bAllowSearch = false;
+	DetailsViewArgs.NameAreaSettings = FDetailsViewArgs::HideNameArea;
+	DetailsViewArgs.bHideSelectionTip = true;
+	DetailsViewArgs.bShowScrollBar = false;
+
+	TSharedRef<IDetailsView> Details = PropertyEditorModule.CreateDetailView(DetailsViewArgs);
+
+	Details->SetObjects(TArray{InObject});
+
+	return Details;
+}
+
 #undef LOCTEXT_NAMESPACE
diff --git a/Engine/Plugins/Experimental/ConsoleVariables/Source/ConsoleVariablesEditor/Private/Views/MainPanel/SConsoleVariablesEditorMainPanel.h b/Engine/Plugins/Experimental/ConsoleVariables/Source/ConsoleVariablesEditor/Private/Views/MainPanel/SConsoleVariablesEditorMainPanel.h
--- a/Engine/Plugins/Experimental/ConsoleVariables/Source/ConsoleVariablesEditor/Private/Views/MainPanel/SConsoleVariablesEditorMainPanel.h
+++ b/Engine/Plugins/Experimental/ConsoleVariables/Source/ConsoleVariablesEditor/Private/Views/MainPanel/SConsoleVariablesEditorMainPanel.h
@@ -6,7 +6,10 @@
 #include "Widgets/SCompoundWidget.h"
 
 class FConsoleVariablesEditorMainPanel;
+class SBox;
 class SMultiLineEditableTextBox;
+class SHorizontalBox;
+class SVerticalBox;
 class UConsoleVariablesAsset;
 
 class SConsoleVariablesEditorMainPanel final : public SCompoundWidget
@@ -26,13 +29,18 @@ public:
 	/** A handler for when text is committed to the Console Input Box that appears when clicking the "Add Console Variable" button */
 	FReply HandleConsoleInputTextCommitted(const FGeometry& MyGeometry, const FKeyEvent& KeyPressed);
 
+	void RefreshMultiUserDetails();
+	void ToggleMultiUserDetails(ECheckBoxState CheckState);
+
 private:
 
 	/** A reference to the struct that controls this widget */
 	TWeakPtr<FConsoleVariablesEditorMainPanel> MainPanel;
 
+	TSharedPtr<SHorizontalBox> ToolbarHBox;
+
 	/** A reference to the button which opens the plugin settings */
-	TSharedPtr<SCheckBox> SettingsButtonPtr;
+	TSharedPtr<SCheckBox> ConcertButtonPtr;
 
 	/** A special text box widget that can search console variables as text is typed into it */
 	TSharedPtr<SWidget> ConsoleInput;
@@ -41,6 +49,8 @@ private:
 
 	/** Creates the toolbar at the top of the MainPanel widget */
 	TSharedRef<SWidget> GeneratePanelToolbar(const TSharedRef<SWidget> InConsoleInputWidget);
+	
+	void CreateConcertButtonIfNeeded();
 
 	// Save / Load
 
@@ -53,4 +63,9 @@ private:
 
 	/** Creates a special asset picker widget to display when the Save/Load button is clicked */
 	TSharedRef<SWidget> OnGeneratePresetsMenu();
+
+	TSharedPtr<SVerticalBox> MultiUserDetailsBox;
+
+	/** Generates a customized details widget given an object. Intended for multi-user settings display. */
+	static TSharedRef<SWidget> GetConcertDetailsWidget(UObject* InObject);
 };
