Console Variables Editor:

-Editing unchecked rows' values is now disabled

#[fyi] Jason.Walter

diff --git a/Engine/Plugins/Experimental/ConsoleVariables/Source/ConsoleVariablesEditor/Private/Views/List/SConsoleVariablesEditorListRow.cpp b/Engine/Plugins/Experimental/ConsoleVariables/Source/ConsoleVariablesEditor/Private/Views/List/SConsoleVariablesEditorListRow.cpp
--- a/Engine/Plugins/Experimental/ConsoleVariables/Source/ConsoleVariablesEditor/Private/Views/List/SConsoleVariablesEditorListRow.cpp
+++ b/Engine/Plugins/Experimental/ConsoleVariables/Source/ConsoleVariablesEditor/Private/Views/List/SConsoleVariablesEditorListRow.cpp
@@ -289,9 +289,9 @@ TSharedRef<SWidget> SConsoleVariablesEditorListRow::GenerateValueCellWidget(cons
 			.ContentPadding(0)
 			.Visibility_Lambda([this]()
 			{
-				if (Item.Pin()->GetRowType() == FConsoleVariablesEditorListRow::SingleCommand)
+				if (const TSharedPtr<FConsoleVariablesEditorListRow> PinnedItem = Item.Pin(); PinnedItem->GetRowType() == FConsoleVariablesEditorListRow::SingleCommand)
 				{
-					return Item.Pin()->GetCommandInfo().Pin()->IsCurrentValueDifferentFromInputValue(Item.Pin()->GetPresetValue()) ?
+					return PinnedItem->IsRowChecked() && Item.Pin()->GetCommandInfo().Pin()->IsCurrentValueDifferentFromInputValue(Item.Pin()->GetPresetValue()) ?
 						EVisibility::Visible : EVisibility::Hidden;
 				}
 
diff --git a/Engine/Plugins/Experimental/ConsoleVariables/Source/ConsoleVariablesEditor/Private/Views/List/SConsoleVariablesEditorListValueInput.cpp b/Engine/Plugins/Experimental/ConsoleVariables/Source/ConsoleVariablesEditor/Private/Views/List/SConsoleVariablesEditorListValueInput.cpp
--- a/Engine/Plugins/Experimental/ConsoleVariables/Source/ConsoleVariablesEditor/Private/Views/List/SConsoleVariablesEditorListValueInput.cpp
+++ b/Engine/Plugins/Experimental/ConsoleVariables/Source/ConsoleVariablesEditor/Private/Views/List/SConsoleVariablesEditorListValueInput.cpp
@@ -58,6 +58,11 @@ void SConsoleVariablesEditorListValueInput::SetCachedValue(const FString& NewCac
 	CachedValue = NewCachedValue;
 }
 
+bool SConsoleVariablesEditorListValueInput::IsRowChecked() const
+{
+	return Item.Pin()->IsRowChecked();
+}
+
 void SConsoleVariablesEditorListValueInput_Float::Construct(const FArguments& InArgs,
                                                             const TWeakPtr<FConsoleVariablesEditorListRow> InRow)
 {
@@ -77,7 +82,7 @@ void SConsoleVariablesEditorListValueInput_Float::Construct(const FArguments& In
 		})
 		.OnValueChanged_Lambda([this] (const float InValue)
 		{
-			if (const TSharedPtr<FConsoleVariablesEditorListRow> PinnedItem = Item.Pin(); PinnedItem->GetWidgetCheckedState() == ECheckBoxState::Checked)
+			if (const TSharedPtr<FConsoleVariablesEditorListRow> PinnedItem = Item.Pin())
 			{
 				const FString ValueAsString = FString::SanitizeFloat(InValue);
 				
@@ -88,6 +93,7 @@ void SConsoleVariablesEditorListValueInput_Float::Construct(const FArguments& In
 			
 			SetCachedValue(FString::SanitizeFloat(InValue));
 		})
+		.IsEnabled(this, &SConsoleVariablesEditorListValueInput::IsRowChecked)
 	];
 
 	SetCachedValue(GetInputValueAsString());
@@ -127,7 +133,7 @@ void SConsoleVariablesEditorListValueInput_Int::Construct(const FArguments& InAr
 		})
 		.OnValueChanged_Lambda([this] (const int32 InValue)
 		{
-			if (const TSharedPtr<FConsoleVariablesEditorListRow> PinnedItem = Item.Pin(); PinnedItem->GetWidgetCheckedState() == ECheckBoxState::Checked)
+			if (const TSharedPtr<FConsoleVariablesEditorListRow> PinnedItem = Item.Pin())
 			{
 				const FString ValueAsString = FString::FromInt(InValue);
 				
@@ -138,6 +144,7 @@ void SConsoleVariablesEditorListValueInput_Int::Construct(const FArguments& InAr
 
 			SetCachedValue(FString::FromInt(InValue));
 		})
+		.IsEnabled(this, &SConsoleVariablesEditorListValueInput::IsRowChecked)
 	];
 	
 	SetCachedValue(GetInputValueAsString());
@@ -177,7 +184,7 @@ void SConsoleVariablesEditorListValueInput_String::Construct(const FArguments& I
 		})
 		.OnTextCommitted_Lambda([this] (const FText& InValue, ETextCommit::Type InTextCommitType)
 		{
-			if (const TSharedPtr<FConsoleVariablesEditorListRow> PinnedItem = Item.Pin(); PinnedItem->GetWidgetCheckedState() == ECheckBoxState::Checked)
+			if (const TSharedPtr<FConsoleVariablesEditorListRow> PinnedItem = Item.Pin())
 			{
 				const FString ValueAsString = InValue.ToString();
 				
@@ -188,6 +195,7 @@ void SConsoleVariablesEditorListValueInput_String::Construct(const FArguments& I
 
 			SetCachedValue(InValue.ToString());
 		})
+		.IsEnabled(this, &SConsoleVariablesEditorListValueInput::IsRowChecked)
 	];
 
 	SetCachedValue(GetInputValueAsString());
@@ -229,7 +237,7 @@ void SConsoleVariablesEditorListValueInput_Bool::Construct(const FArguments& InA
 		.MaxSliderValue(2)
 		.OnValueChanged_Lambda([this] (const int32 InValue)
 		{
-			if (const TSharedPtr<FConsoleVariablesEditorListRow> PinnedItem = Item.Pin(); PinnedItem->GetWidgetCheckedState() == ECheckBoxState::Checked)
+			if (const TSharedPtr<FConsoleVariablesEditorListRow> PinnedItem = Item.Pin())
 			{
 				const FString ValueAsString = FString::FromInt(InValue);
 				
@@ -240,6 +248,7 @@ void SConsoleVariablesEditorListValueInput_Bool::Construct(const FArguments& InA
 
 			SetCachedValue(FString::FromInt(InValue));
 		})
+		.IsEnabled(this, &SConsoleVariablesEditorListValueInput::IsRowChecked)
 	];
 	SetCachedValue(GetInputValueAsString());
 }
diff --git a/Engine/Plugins/Experimental/ConsoleVariables/Source/ConsoleVariablesEditor/Private/Views/List/SConsoleVariablesEditorListValueInput.h b/Engine/Plugins/Experimental/ConsoleVariables/Source/ConsoleVariablesEditor/Private/Views/List/SConsoleVariablesEditorListValueInput.h
--- a/Engine/Plugins/Experimental/ConsoleVariables/Source/ConsoleVariablesEditor/Private/Views/List/SConsoleVariablesEditorListValueInput.h
+++ b/Engine/Plugins/Experimental/ConsoleVariables/Source/ConsoleVariablesEditor/Private/Views/List/SConsoleVariablesEditorListValueInput.h
@@ -20,6 +20,8 @@ public:
 	[[nodiscard]] const FString& GetCachedValue() const;
 	void SetCachedValue(const FString& CachedValue);
 
+	bool IsRowChecked() const;
+
 protected:
 	
 	TWeakPtr<FConsoleVariablesEditorListRow> Item = nullptr;
