Console Variables Editor:

-Add project setting to choose which value to display when the row is unchecked

#[fyi] Jason.Walter

diff --git a/Engine/Plugins/Experimental/ConsoleVariables/Source/ConsoleVariablesEditor/Private/Views/List/SConsoleVariablesEditorListValueInput.cpp b/Engine/Plugins/Experimental/ConsoleVariables/Source/ConsoleVariablesEditor/Private/Views/List/SConsoleVariablesEditorListValueInput.cpp
--- a/Engine/Plugins/Experimental/ConsoleVariables/Source/ConsoleVariablesEditor/Private/Views/List/SConsoleVariablesEditorListValueInput.cpp
+++ b/Engine/Plugins/Experimental/ConsoleVariables/Source/ConsoleVariablesEditor/Private/Views/List/SConsoleVariablesEditorListValueInput.cpp
@@ -4,6 +4,7 @@
 
 #include "ConsoleVariablesEditorListRow.h"
 #include "ConsoleVariablesEditorModule.h"
+#include "ConsoleVariablesEditorProjectSettings.h"
 
 #include "Widgets/Input/SNumericEntryBox.h"
 
@@ -67,13 +68,18 @@ void SConsoleVariablesEditorListValueInput_Float::Construct(const FArguments& In
                                                             const TWeakPtr<FConsoleVariablesEditorListRow> InRow)
 {
 	Item = InRow;
+
+	ProjectSettingsPtr = GetMutableDefault<UConsoleVariablesEditorProjectSettings>();
 	
 	ChildSlot
 	[
 		SAssignNew(InputWidget, SSpinBox<float>)
+		.MaxFractionalDigits(3)
 		.Value_Lambda([this]
 		{
-			if (Item.IsValid() && Item.Pin()->GetWidgetCheckedState() == ECheckBoxState::Checked)
+			if (Item.IsValid() &&
+				(Item.Pin()->GetWidgetCheckedState() == ECheckBoxState::Checked ||
+					(ProjectSettingsPtr && ProjectSettingsPtr->UncheckedRowDisplayType == EConsoleVariablesEditorRowDisplayType::ShowCurrentValue)))
 			{
 				return FCString::Atof(*Item.Pin()->GetCommandInfo().Pin()->ConsoleVariablePtr->GetString());
 			}
@@ -118,13 +124,17 @@ void SConsoleVariablesEditorListValueInput_Int::Construct(const FArguments& InAr
                                                           const TWeakPtr<FConsoleVariablesEditorListRow> InRow)
 {
 	Item = InRow;
+
+	ProjectSettingsPtr = GetMutableDefault<UConsoleVariablesEditorProjectSettings>();
 	
 	ChildSlot
 	[
 		SAssignNew(InputWidget, SSpinBox<int32>)
 		.Value_Lambda([this]
 		{
-			if (Item.IsValid() && Item.Pin()->GetWidgetCheckedState() == ECheckBoxState::Checked)
+			if (Item.IsValid() &&
+				(Item.Pin()->GetWidgetCheckedState() == ECheckBoxState::Checked ||
+					(ProjectSettingsPtr && ProjectSettingsPtr->UncheckedRowDisplayType == EConsoleVariablesEditorRowDisplayType::ShowCurrentValue)))
 			{
 				return FCString::Atoi(*Item.Pin()->GetCommandInfo().Pin()->ConsoleVariablePtr->GetString());
 			}
@@ -169,13 +179,17 @@ void SConsoleVariablesEditorListValueInput_String::Construct(const FArguments& I
                                                              const TWeakPtr<FConsoleVariablesEditorListRow> InRow)
 {
 	Item = InRow;
+
+	ProjectSettingsPtr = GetMutableDefault<UConsoleVariablesEditorProjectSettings>();
 	
 	ChildSlot
 	[
 		SAssignNew(InputWidget, SEditableText)
 		.Text_Lambda([this]
 		{
-			if (Item.IsValid() && Item.Pin()->GetWidgetCheckedState() == ECheckBoxState::Checked)
+			if (Item.IsValid() &&
+				(Item.Pin()->GetWidgetCheckedState() == ECheckBoxState::Checked ||
+					(ProjectSettingsPtr && ProjectSettingsPtr->UncheckedRowDisplayType == EConsoleVariablesEditorRowDisplayType::ShowCurrentValue)))
 			{
 				return FText::FromString(*Item.Pin()->GetCommandInfo().Pin()->ConsoleVariablePtr->GetString());
 			}
@@ -220,13 +234,17 @@ void SConsoleVariablesEditorListValueInput_Bool::Construct(const FArguments& InA
                                                            const TWeakPtr<FConsoleVariablesEditorListRow> InRow)
 {
 	Item = InRow;
+
+	ProjectSettingsPtr = GetMutableDefault<UConsoleVariablesEditorProjectSettings>();
 	
 	ChildSlot
 	[
 		SAssignNew(InputWidget, SSpinBox<int32>)
 		.Value_Lambda([this]
 		{
-			if (Item.IsValid() && Item.Pin()->GetWidgetCheckedState() == ECheckBoxState::Checked)
+			if (Item.IsValid() &&
+				(Item.Pin()->GetWidgetCheckedState() == ECheckBoxState::Checked ||
+					(ProjectSettingsPtr && ProjectSettingsPtr->UncheckedRowDisplayType == EConsoleVariablesEditorRowDisplayType::ShowCurrentValue)))
 			{
 				return FCString::Atoi(*Item.Pin()->GetCommandInfo().Pin()->ConsoleVariablePtr->GetString());
 			}
diff --git a/Engine/Plugins/Experimental/ConsoleVariables/Source/ConsoleVariablesEditor/Private/Views/List/SConsoleVariablesEditorListValueInput.h b/Engine/Plugins/Experimental/ConsoleVariables/Source/ConsoleVariablesEditor/Private/Views/List/SConsoleVariablesEditorListValueInput.h
--- a/Engine/Plugins/Experimental/ConsoleVariables/Source/ConsoleVariablesEditor/Private/Views/List/SConsoleVariablesEditorListValueInput.h
+++ b/Engine/Plugins/Experimental/ConsoleVariables/Source/ConsoleVariablesEditor/Private/Views/List/SConsoleVariablesEditorListValueInput.h
@@ -7,6 +7,8 @@
 
 struct FConsoleVariablesEditorListRow;
 
+class UConsoleVariablesEditorProjectSettings;
+
 class SConsoleVariablesEditorListValueInput : public SCompoundWidget
 {
 public:
@@ -26,6 +28,8 @@ protected:
 	
 	TWeakPtr<FConsoleVariablesEditorListRow> Item = nullptr;
 	FString CachedValue = "";
+
+	TObjectPtr<UConsoleVariablesEditorProjectSettings> ProjectSettingsPtr = nullptr;
 };
 
 class SConsoleVariablesEditorListValueInput_Float : public SConsoleVariablesEditorListValueInput
diff --git a/Engine/Plugins/Experimental/ConsoleVariables/Source/ConsoleVariablesEditor/Public/ConsoleVariablesEditorProjectSettings.h b/Engine/Plugins/Experimental/ConsoleVariables/Source/ConsoleVariablesEditor/Public/ConsoleVariablesEditorProjectSettings.h
--- a/Engine/Plugins/Experimental/ConsoleVariables/Source/ConsoleVariablesEditor/Public/ConsoleVariablesEditorProjectSettings.h
+++ b/Engine/Plugins/Experimental/ConsoleVariables/Source/ConsoleVariablesEditor/Public/ConsoleVariablesEditorProjectSettings.h
@@ -4,6 +4,13 @@
 
 #include "ConsoleVariablesEditorProjectSettings.generated.h"
 
+UENUM(BlueprintType)
+enum class EConsoleVariablesEditorRowDisplayType : uint8
+{
+	ShowCurrentValue,
+	ShowLastEnteredValue
+};
+
 UCLASS(config = Engine, defaultconfig)
 class CONSOLEVARIABLESEDITOR_API UConsoleVariablesEditorProjectSettings : public UObject
 {
@@ -18,4 +25,11 @@ public:
 	/** If true, any console variable changes will be added to the current preset as long as the plugin is loaded. */
 	UPROPERTY(Config, EditAnywhere, Category="Console Variables")
 	bool bAddAllChangedConsoleVariablesToCurrentPreset;
+
+	/** When a row is unchecked, its associated variable's value will be set to the value recorded when the plugin was loaded.
+	 *The value displayed to the user can be configured with this setting, but will not affect the actual applied value.
+	 *ShowCurrentValue displays the actual value currently applied to the variable.
+	 *ShowLastEnteredValue displays the value that will be applied when the row is checked. */
+	UPROPERTY(Config, EditAnywhere, Category="Console Variables")
+	EConsoleVariablesEditorRowDisplayType UncheckedRowDisplayType;
 };
