Console Variables Editor:

-Unchecked Rows now display startup source

#[fyi] Jason.Walter

diff --git a/Engine/Plugins/Experimental/ConsoleVariables/Source/ConsoleVariablesEditor/Private/ConsoleVariablesEditorCommandInfo.h b/Engine/Plugins/Experimental/ConsoleVariables/Source/ConsoleVariablesEditor/Private/ConsoleVariablesEditorCommandInfo.h
--- a/Engine/Plugins/Experimental/ConsoleVariables/Source/ConsoleVariablesEditor/Private/ConsoleVariablesEditorCommandInfo.h
+++ b/Engine/Plugins/Experimental/ConsoleVariables/Source/ConsoleVariablesEditor/Private/ConsoleVariablesEditorCommandInfo.h
@@ -55,13 +55,21 @@ struct FConsoleVariablesEditorCommandInfo
 		return CurrentWorld;
 	}
 
-	FText GetSource() const
+	EConsoleVariableFlags GetSource() const
+	{
+		return (EConsoleVariableFlags)((uint32)ConsoleVariablePtr->GetFlags() & ECVF_SetByMask);
+	}
+
+	void SetSourceFlag(const EConsoleVariableFlags InSource) const
+	{
+		ConsoleVariablePtr->SetFlags((EConsoleVariableFlags)InSource);
+	}
+
+	FText GetSourceAsText() const
 	{
 		if (ConsoleVariablePtr)
 		{
-			const EConsoleVariableFlags SetBy = (EConsoleVariableFlags)((uint32)ConsoleVariablePtr->GetFlags() & ECVF_SetByMask);
-
-			switch (SetBy)
+			switch (GetSource())
 			{
 			case (EConsoleVariableFlags::ECVF_SetByConstructor):
 				return LOCTEXT("SetByConstructor", "Constructor");
@@ -120,8 +128,10 @@ struct FConsoleVariablesEditorCommandInfo
 
 	IConsoleVariable* ConsoleVariablePtr;
 
-	/** The value of this command when the module started in this session after it may have been set by an ini file */
+	/** The value of this command when the module started in this session after it may have been set by an ini file. */
 	FString StartupValueAsString;
+	/** The source of this variable last setting as recorded when the plugin was loaded. */
+	EConsoleVariableFlags StartupSource;
 	
 	FDelegateHandle OnVariableChangedCallbackHandle;
 };
diff --git a/Engine/Plugins/Experimental/ConsoleVariables/Source/ConsoleVariablesEditor/Private/ConsoleVariablesEditorModule.cpp b/Engine/Plugins/Experimental/ConsoleVariables/Source/ConsoleVariablesEditor/Private/ConsoleVariablesEditorModule.cpp
--- a/Engine/Plugins/Experimental/ConsoleVariables/Source/ConsoleVariablesEditor/Private/ConsoleVariablesEditorModule.cpp
+++ b/Engine/Plugins/Experimental/ConsoleVariables/Source/ConsoleVariablesEditor/Private/ConsoleVariablesEditorModule.cpp
@@ -89,9 +89,12 @@ void FConsoleVariablesEditorModule::QueryAndBeginTrackingConsoleVariables()
 		{
 			if (IConsoleVariable* AsVariable = ConsoleObject->AsVariable())
 			{
-				const FDelegateHandle Handle = AsVariable->OnChangedDelegate().AddRaw(this, &FConsoleVariablesEditorModule::OnConsoleVariableChanged);
-				ConsoleVariablesMasterReference.Add(
-					MakeShared<FConsoleVariablesEditorCommandInfo>(Key, AsVariable, AsVariable->GetString(), Handle));
+				const FDelegateHandle Handle =
+					AsVariable->OnChangedDelegate().AddRaw(this, &FConsoleVariablesEditorModule::OnConsoleVariableChanged);
+				const TSharedRef<FConsoleVariablesEditorCommandInfo> Info =
+					MakeShared<FConsoleVariablesEditorCommandInfo>(Key, AsVariable, AsVariable->GetString(), Handle);
+				Info->StartupSource = Info->GetSource();
+				ConsoleVariablesMasterReference.Add(Info);
 			}
 		}),
 		TEXT(""));
diff --git a/Engine/Plugins/Experimental/ConsoleVariables/Source/ConsoleVariablesEditor/Private/Views/List/ConsoleVariablesEditorListRow.cpp b/Engine/Plugins/Experimental/ConsoleVariables/Source/ConsoleVariablesEditor/Private/Views/List/ConsoleVariablesEditorListRow.cpp
--- a/Engine/Plugins/Experimental/ConsoleVariables/Source/ConsoleVariablesEditor/Private/Views/List/ConsoleVariablesEditorListRow.cpp
+++ b/Engine/Plugins/Experimental/ConsoleVariables/Source/ConsoleVariablesEditor/Private/Views/List/ConsoleVariablesEditorListRow.cpp
@@ -108,7 +108,7 @@ bool FConsoleVariablesEditorListRow::MatchSearchTokensToSearchTerms(const TArray
 	{
 		TSharedPtr<FConsoleVariablesEditorCommandInfo> PinnedInfo = CommandInfo.Pin();
 		
-		const FString SearchTerms = PinnedInfo->Command + PinnedInfo->GetSource().ToString() +
+		const FString SearchTerms = PinnedInfo->Command + PinnedInfo->GetSourceAsText().ToString() +
 			(PinnedInfo->ConsoleVariablePtr ? PinnedInfo->ConsoleVariablePtr->GetString() + PinnedInfo->ConsoleVariablePtr->GetHelp() : "");
 
 		bMatchFound = Algo::AnyOf(InTokens,
diff --git a/Engine/Plugins/Experimental/ConsoleVariables/Source/ConsoleVariablesEditor/Private/Views/List/SConsoleVariablesEditorList.h b/Engine/Plugins/Experimental/ConsoleVariables/Source/ConsoleVariablesEditor/Private/Views/List/SConsoleVariablesEditorList.h
--- a/Engine/Plugins/Experimental/ConsoleVariables/Source/ConsoleVariablesEditor/Private/Views/List/SConsoleVariablesEditorList.h
+++ b/Engine/Plugins/Experimental/ConsoleVariables/Source/ConsoleVariablesEditor/Private/Views/List/SConsoleVariablesEditorList.h
@@ -129,12 +129,12 @@ private:
 	TFunctionRef<bool(const FConsoleVariablesEditorListRowPtr&, const FConsoleVariablesEditorListRowPtr&)> SortBySourceAscending =
 		[](const FConsoleVariablesEditorListRowPtr& A, const FConsoleVariablesEditorListRowPtr& B)
 	{
-		return A->GetCommandInfo().Pin()->GetSource().ToString() < B->GetCommandInfo().Pin()->GetSource().ToString();
+		return A->GetCommandInfo().Pin()->GetSourceAsText().ToString() < B->GetCommandInfo().Pin()->GetSourceAsText().ToString();
 	};
 
 	TFunctionRef<bool(const FConsoleVariablesEditorListRowPtr&, const FConsoleVariablesEditorListRowPtr&)> SortBySourceDescending =
 		[](const FConsoleVariablesEditorListRowPtr& A, const FConsoleVariablesEditorListRowPtr& B)
 	{
-		return B->GetCommandInfo().Pin()->GetSource().ToString() < A->GetCommandInfo().Pin()->GetSource().ToString();
+		return B->GetCommandInfo().Pin()->GetSourceAsText().ToString() < A->GetCommandInfo().Pin()->GetSourceAsText().ToString();
 	};
 };
diff --git a/Engine/Plugins/Experimental/ConsoleVariables/Source/ConsoleVariablesEditor/Private/Views/List/SConsoleVariablesEditorListRow.cpp b/Engine/Plugins/Experimental/ConsoleVariables/Source/ConsoleVariablesEditor/Private/Views/List/SConsoleVariablesEditorListRow.cpp
--- a/Engine/Plugins/Experimental/ConsoleVariables/Source/ConsoleVariablesEditor/Private/Views/List/SConsoleVariablesEditorListRow.cpp
+++ b/Engine/Plugins/Experimental/ConsoleVariables/Source/ConsoleVariablesEditor/Private/Views/List/SConsoleVariablesEditorListRow.cpp
@@ -101,28 +101,6 @@ void SConsoleVariablesEditorListRow::OnMouseLeave(const FPointerEvent& MouseEven
 
 SConsoleVariablesEditorListRow::~SConsoleVariablesEditorListRow()
 {
-	// Remove delegate bindings
-
-	// Unbind event to the splitter being resized first
-	if (NestedSplitterPtr.IsValid())
-	{
-		for (int32 SplitterSlotCount = 0; SplitterSlotCount < NestedSplitterPtr->GetChildren()->Num(); SplitterSlotCount++)
-		{
-			NestedSplitterPtr->SlotAt(SplitterSlotCount).OnSlotResized().Unbind();
-		}
-	}
-
-	if (OuterSplitterPtr.IsValid())
-	{
-		for (int32 SplitterSlotCount = 0; SplitterSlotCount < OuterSplitterPtr->GetChildren()->Num(); SplitterSlotCount++)
-		{
-			OuterSplitterPtr->SlotAt(SplitterSlotCount).OnSlotResized().Unbind();
-		}
-	}
-
-	OuterSplitterPtr.Reset();
-	NestedSplitterPtr.Reset();
-
 	Item.Reset();
 
 	FlashImages.Empty();
@@ -205,7 +183,7 @@ TSharedRef<SWidget> SConsoleVariablesEditorListRow::GenerateCells(const FName& I
 					SNew(STextBlock)
 					.Text_Lambda([this]()
 					{
-						return Item.Pin()->GetCommandInfo().Pin()->GetSource();
+						return Item.Pin()->GetCommandInfo().Pin()->GetSourceAsText();
 					})
 				]
 
@@ -252,6 +230,8 @@ void SConsoleVariablesEditorListRow::OnCheckboxStateChange(const ECheckBoxState
 
 				FConsoleVariablesEditorModule::Get().SendMultiUserConsoleVariableChange(
 					PinnedItem->GetCommandInfo().Pin()->Command, PinnedItem->GetCommandInfo().Pin()->StartupValueAsString);
+
+				PinnedItem->GetCommandInfo().Pin()->SetSourceFlag(PinnedItem->GetCommandInfo().Pin()->StartupSource);
 			}
 		}
 
diff --git a/Engine/Plugins/Experimental/ConsoleVariables/Source/ConsoleVariablesEditor/Private/Views/List/SConsoleVariablesEditorListRow.h b/Engine/Plugins/Experimental/ConsoleVariables/Source/ConsoleVariablesEditor/Private/Views/List/SConsoleVariablesEditorListRow.h
--- a/Engine/Plugins/Experimental/ConsoleVariables/Source/ConsoleVariablesEditor/Private/Views/List/SConsoleVariablesEditorListRow.h
+++ b/Engine/Plugins/Experimental/ConsoleVariables/Source/ConsoleVariablesEditor/Private/Views/List/SConsoleVariablesEditorListRow.h
@@ -50,9 +50,6 @@ public:
 
 private:
 	
-	TSharedPtr<SSplitter> OuterSplitterPtr;
-	TSharedPtr<SSplitter> NestedSplitterPtr;
-	
 	TWeakPtr<FConsoleVariablesEditorListRow> Item;
 
 	TArray<TSharedPtr<SImage>> FlashImages;
