Console Variables Editor:

-Revert button tooltip now shows preset value
-Value Widget tooltip now shows preset, startup, and custom values

#[fyi] Jason.Walter

diff --git a/Engine/Plugins/Experimental/ConsoleVariables/Source/ConsoleVariablesEditor/Private/ConsoleVariablesEditorCommandInfo.h b/Engine/Plugins/Experimental/ConsoleVariables/Source/ConsoleVariablesEditor/Private/ConsoleVariablesEditorCommandInfo.h
--- a/Engine/Plugins/Experimental/ConsoleVariables/Source/ConsoleVariablesEditor/Private/ConsoleVariablesEditorCommandInfo.h
+++ b/Engine/Plugins/Experimental/ConsoleVariables/Source/ConsoleVariablesEditor/Private/ConsoleVariablesEditorCommandInfo.h
@@ -57,7 +57,12 @@ struct FConsoleVariablesEditorCommandInfo
 
 	EConsoleVariableFlags GetSource() const
 	{
-		return (EConsoleVariableFlags)((uint32)ConsoleVariablePtr->GetFlags() & ECVF_SetByMask);
+		if (ConsoleVariablePtr)
+		{
+			return (EConsoleVariableFlags)((uint32)ConsoleVariablePtr->GetFlags() & ECVF_SetByMask);
+		}
+
+		return ECVF_Default;
 	}
 
 	void SetSourceFlag(const EConsoleVariableFlags InSource) const
@@ -67,10 +72,14 @@ struct FConsoleVariablesEditorCommandInfo
 
 	FText GetSourceAsText() const
 	{
-		if (ConsoleVariablePtr)
+		return ConvertConsoleVariableSetByFlagToText(GetSource());
+	}
+
+	static FText ConvertConsoleVariableSetByFlagToText(const EConsoleVariableFlags InFlag)
+	{
+		
+		switch (InFlag)
 		{
-			switch (GetSource())
-			{
 			case (EConsoleVariableFlags::ECVF_SetByConstructor):
 				return LOCTEXT("SetByConstructor", "Constructor");
 				
@@ -106,7 +115,6 @@ struct FConsoleVariablesEditorCommandInfo
 
 			default:
 				break;;
-			}
 		}
 		
 		return LOCTEXT("UnknownSource", "<UNKNOWN>");
diff --git a/Engine/Plugins/Experimental/ConsoleVariables/Source/ConsoleVariablesEditor/Private/Views/List/ConsoleVariablesEditorListRow.cpp b/Engine/Plugins/Experimental/ConsoleVariables/Source/ConsoleVariablesEditor/Private/Views/List/ConsoleVariablesEditorListRow.cpp
--- a/Engine/Plugins/Experimental/ConsoleVariables/Source/ConsoleVariablesEditor/Private/Views/List/ConsoleVariablesEditorListRow.cpp
+++ b/Engine/Plugins/Experimental/ConsoleVariables/Source/ConsoleVariablesEditor/Private/Views/List/ConsoleVariablesEditorListRow.cpp
@@ -98,8 +98,14 @@ void FConsoleVariablesEditorListRow::SetShouldExpandAllChildren(const bool bNewS
 	bShouldExpandAllChildren = bNewShouldExpandAllChildren;
 }
 
+void FConsoleVariablesEditorListRow::ResetToStartupValueAndSource() const
+{
+	GetCommandInfo().Pin()->ExecuteCommand(GetCommandInfo().Pin()->StartupValueAsString);
+	GetCommandInfo().Pin()->SetSourceFlag(GetCommandInfo().Pin()->StartupSource);
+}
+
 bool FConsoleVariablesEditorListRow::MatchSearchTokensToSearchTerms(const TArray<FString> InTokens,
-	const bool bMatchAnyTokens)
+                                                                    const bool bMatchAnyTokens)
 {
 	// If the search is cleared we'll consider the row to pass search
 	bool bMatchFound = InTokens.Num() == 0;
@@ -188,7 +194,7 @@ FReply FConsoleVariablesEditorListRow::OnRemoveButtonClicked()
 		return FReply::Handled();
 	}
 
-	GetCommandInfo().Pin()->ExecuteCommand(GetCommandInfo().Pin()->StartupValueAsString);
+	ResetToStartupValueAndSource();
 
 	const FConsoleVariablesEditorModule& ConsoleVariablesEditorModule = FConsoleVariablesEditorModule::Get();
 
diff --git a/Engine/Plugins/Experimental/ConsoleVariables/Source/ConsoleVariablesEditor/Private/Views/List/ConsoleVariablesEditorListRow.h b/Engine/Plugins/Experimental/ConsoleVariables/Source/ConsoleVariablesEditor/Private/Views/List/ConsoleVariablesEditorListRow.h
--- a/Engine/Plugins/Experimental/ConsoleVariables/Source/ConsoleVariablesEditor/Private/Views/List/ConsoleVariablesEditorListRow.h
+++ b/Engine/Plugins/Experimental/ConsoleVariables/Source/ConsoleVariablesEditor/Private/Views/List/ConsoleVariablesEditorListRow.h
@@ -63,6 +63,8 @@ struct FConsoleVariablesEditorListRow final : TSharedFromThis<FConsoleVariablesE
 
 	[[nodiscard]] bool GetShouldExpandAllChildren() const;
 	void SetShouldExpandAllChildren(const bool bNewShouldExpandAllChildren);
+
+	void ResetToStartupValueAndSource() const;
 	
 	[[nodiscard]] const FString& GetPresetValue() const;
 
diff --git a/Engine/Plugins/Experimental/ConsoleVariables/Source/ConsoleVariablesEditor/Private/Views/List/SConsoleVariablesEditorListRow.cpp b/Engine/Plugins/Experimental/ConsoleVariables/Source/ConsoleVariablesEditor/Private/Views/List/SConsoleVariablesEditorListRow.cpp
--- a/Engine/Plugins/Experimental/ConsoleVariables/Source/ConsoleVariablesEditor/Private/Views/List/SConsoleVariablesEditorListRow.cpp
+++ b/Engine/Plugins/Experimental/ConsoleVariables/Source/ConsoleVariablesEditor/Private/Views/List/SConsoleVariablesEditorListRow.cpp
@@ -107,8 +107,6 @@ SConsoleVariablesEditorListRow::~SConsoleVariablesEditorListRow()
 
 	ValueChildInputWidget.Reset();
 	
-	SplitterManagerPtr.Reset();
-	
 	HoverableWidgetsPtr.Reset();
 }
 
@@ -226,12 +224,10 @@ void SConsoleVariablesEditorListRow::OnCheckboxStateChange(const ECheckBoxState
 			}
 			else
 			{
-				PinnedItem->GetCommandInfo().Pin()->ExecuteCommand(PinnedItem->GetCommandInfo().Pin()->StartupValueAsString);
+				PinnedItem->ResetToStartupValueAndSource();
 
 				FConsoleVariablesEditorModule::Get().SendMultiUserConsoleVariableChange(
 					PinnedItem->GetCommandInfo().Pin()->Command, PinnedItem->GetCommandInfo().Pin()->StartupValueAsString);
-
-				PinnedItem->GetCommandInfo().Pin()->SetSourceFlag(PinnedItem->GetCommandInfo().Pin()->StartupSource);
 			}
 		}
 
@@ -246,8 +242,19 @@ TSharedRef<SWidget> SConsoleVariablesEditorListRow::GenerateValueCellWidget(cons
 {
 	if (PinnedItem->GetCommandInfo().IsValid() && PinnedItem->GetCommandInfo().Pin()->ConsoleVariablePtr)
 	{
-		const TSharedRef<SHorizontalBox> FinalValueWidget = SNew(SHorizontalBox);
 		ValueChildInputWidget = SConsoleVariablesEditorListValueInput::GetInputWidget(Item);
+		const TSharedRef<SHorizontalBox> FinalValueWidget =
+			SNew(SHorizontalBox)
+			.ToolTipText_Lambda([this, PinnedItem] ()
+			{
+				return FText::Format(
+							ValueWidgetToolTipFormatText,
+							FText::FromString(ValueChildInputWidget->GetCachedValue()),
+							FText::FromString(PinnedItem->GetPresetValue()),
+							FText::FromString(PinnedItem->GetCommandInfo().Pin()->StartupValueAsString),
+							FConsoleVariablesEditorCommandInfo::ConvertConsoleVariableSetByFlagToText(PinnedItem->GetCommandInfo().Pin()->StartupSource)
+				);
+			});
 
 		FinalValueWidget->AddSlot()
 		.VAlign(VAlign_Center)
@@ -263,23 +270,22 @@ TSharedRef<SWidget> SConsoleVariablesEditorListRow::GenerateValueCellWidget(cons
 		[
 			SNew(SButton)
 			.IsFocusable(false)
-			.ToolTipText(
-				LOCTEXT("ResetRowValueTooltipText", "Reset this value to what is defined in the preset or what it was when the engine started."))
+			.ToolTipText(FText::Format(RevertButtonFormatText, FText::FromString(PinnedItem->GetPresetValue())))
 			.ButtonStyle(&FAppStyle::Get().GetWidgetStyle<FButtonStyle>("NoBorder"))
 			.ContentPadding(0)
-			.Visibility_Lambda([this]()
+			.Visibility_Lambda([this, PinnedItem]()
 			{
-				if (const TSharedPtr<FConsoleVariablesEditorListRow> PinnedItem = Item.Pin(); PinnedItem->GetRowType() == FConsoleVariablesEditorListRow::SingleCommand)
+				if (PinnedItem.IsValid() && PinnedItem->GetRowType() == FConsoleVariablesEditorListRow::SingleCommand)
 				{
-					return PinnedItem->IsRowChecked() && Item.Pin()->GetCommandInfo().Pin()->IsCurrentValueDifferentFromInputValue(Item.Pin()->GetPresetValue()) ?
+					return PinnedItem->IsRowChecked() && PinnedItem->GetCommandInfo().Pin()->IsCurrentValueDifferentFromInputValue(PinnedItem->GetPresetValue()) ?
 						EVisibility::Visible : EVisibility::Hidden;
 				}
 
 				return EVisibility::Collapsed;
 			})
-			.OnClicked_Lambda([this]()
+			.OnClicked_Lambda([this, PinnedItem]()
 			{
-				Item.Pin()->ResetToPresetValue();
+				PinnedItem->ResetToPresetValue();
 
 				return FReply::Handled();
 			})
diff --git a/Engine/Plugins/Experimental/ConsoleVariables/Source/ConsoleVariablesEditor/Private/Views/List/SConsoleVariablesEditorListRow.h b/Engine/Plugins/Experimental/ConsoleVariables/Source/ConsoleVariablesEditor/Private/Views/List/SConsoleVariablesEditorListRow.h
--- a/Engine/Plugins/Experimental/ConsoleVariables/Source/ConsoleVariablesEditor/Private/Views/List/SConsoleVariablesEditorListRow.h
+++ b/Engine/Plugins/Experimental/ConsoleVariables/Source/ConsoleVariablesEditor/Private/Views/List/SConsoleVariablesEditorListRow.h
@@ -56,14 +56,13 @@ private:
 
 	TSharedPtr<SConsoleVariablesEditorListValueInput> ValueChildInputWidget;
 
-	/* For splitter sync */
+	#define LOCTEXT_NAMESPACE "ConsoleVariablesEditor"
 
-	/* To sync up splitter location in tree view items, we need to account for the tree view's indentation.
-	 * Instead of calculating the coefficient twice each frame (for left and right splitter slots), we do it once and cache it here. */
-	float CachedNestedColumnWidthAdjusted = 0.f;
-	
-	FConsoleVariablesEditorListSplitterManagerPtr SplitterManagerPtr;
+	FText ValueWidgetToolTipFormatText = LOCTEXT("ValueWidgetToolTipFormatText", "Custom Value: {0}\nPreset Value: {1}\nStartup Value: {2} (Set By {3})");
+	FText RevertButtonFormatText = LOCTEXT("RevertButtonFormatText", "Reset to Preset Value: {0}");
 
+	#undef LOCTEXT_NAMESPACE
+	
 	TSharedPtr<SConsoleVariablesEditorListRowHoverWidgets> HoverableWidgetsPtr;
 
 	FCurveSequence FlashAnimation;
