Console Variables Editor:

-Implement drag and drop reordering
-Implement Sort Order column/property
-Implement option to set new custom sort order
-Set fill widths for some columns so full name visibility is prioritized
-Miscellaneous terminology and spelling updates
-Stopped value sliders from updating Source unless the value actually changes
-Added skeleton code for adding Show Filters (not implemented yet)
-Removed some old Level Snapshots UI implementations that don't apply to CVars as it is now

diff --git a/Engine/Plugins/Experimental/ConsoleVariables/Source/ConsoleVariablesEditor/Private/ConsoleVariablesEditorStyle.cpp b/Engine/Plugins/Experimental/ConsoleVariables/Source/ConsoleVariablesEditor/Private/ConsoleVariablesEditorStyle.cpp
--- a/Engine/Plugins/Experimental/ConsoleVariables/Source/ConsoleVariablesEditor/Private/ConsoleVariablesEditorStyle.cpp
+++ b/Engine/Plugins/Experimental/ConsoleVariables/Source/ConsoleVariablesEditor/Private/ConsoleVariablesEditorStyle.cpp
@@ -64,7 +64,7 @@ TSharedRef< FSlateStyleSet > FConsoleVariablesEditorStyle::Create()
 
 	Style->SetCoreContentRoot(FPaths::EngineContentDir() / TEXT("Slate"));
 
-	// Toolbar
+	// Icons
 	Style->Set("ConsoleVariables.ToolbarButton", new IMAGE_BRUSH("Icons/Icon40", Icon40x40));
 	Style->Set("ConsoleVariables.ToolbarButton.Small", new IMAGE_BRUSH("Icons/Icon20", Icon20x20));
 
diff --git a/Engine/Plugins/Experimental/ConsoleVariables/Source/ConsoleVariablesEditor/Private/Views/List/ConsoleVariablesEditorList.h b/Engine/Plugins/Experimental/ConsoleVariables/Source/ConsoleVariablesEditor/Private/Views/List/ConsoleVariablesEditorList.h
--- a/Engine/Plugins/Experimental/ConsoleVariables/Source/ConsoleVariablesEditor/Private/Views/List/ConsoleVariablesEditorList.h
+++ b/Engine/Plugins/Experimental/ConsoleVariables/Source/ConsoleVariablesEditor/Private/Views/List/ConsoleVariablesEditorList.h
@@ -4,8 +4,11 @@
 
 #include "Widgets/SWidget.h"
 
-class UConsoleVariablesAsset;
+struct FConsoleVariablesEditorListRow;
+typedef TSharedPtr<FConsoleVariablesEditorListRow> FConsoleVariablesEditorListRowPtr;
+
 class SConsoleVariablesEditorList;
+class UConsoleVariablesAsset;
 
 class FConsoleVariablesEditorList : public TSharedFromThis<FConsoleVariablesEditorList>
 {
diff --git a/Engine/Plugins/Experimental/ConsoleVariables/Source/ConsoleVariablesEditor/Private/Views/List/ConsoleVariablesEditorListRow.cpp b/Engine/Plugins/Experimental/ConsoleVariables/Source/ConsoleVariablesEditor/Private/Views/List/ConsoleVariablesEditorListRow.cpp
--- a/Engine/Plugins/Experimental/ConsoleVariables/Source/ConsoleVariablesEditor/Private/Views/List/ConsoleVariablesEditorListRow.cpp
+++ b/Engine/Plugins/Experimental/ConsoleVariables/Source/ConsoleVariablesEditor/Private/Views/List/ConsoleVariablesEditorListRow.cpp
@@ -41,6 +41,16 @@ void FConsoleVariablesEditorListRow::SetChildDepth(const int32 InDepth)
 	ChildDepth = InDepth;
 }
 
+int32 FConsoleVariablesEditorListRow::GetSortOrder() const
+{
+	return SortOrder;
+}
+
+void FConsoleVariablesEditorListRow::SetSortOrder(const int32 InNewOrder)
+{
+	SortOrder = InNewOrder;
+}
+
 TWeakPtr<FConsoleVariablesEditorListRow> FConsoleVariablesEditorListRow::GetDirectParentRow() const
 {
 	return DirectParentRow;
@@ -124,7 +134,7 @@ bool FConsoleVariablesEditorListRow::MatchSearchTokensToSearchTerms(const TArray
 			});
 	}
 
-	bDoesRowMatchSeachTerms = bMatchFound;
+	bDoesRowMatchSearchTerms = bMatchFound;
 
 	return bMatchFound;
 }
@@ -167,6 +177,16 @@ void FConsoleVariablesEditorListRow::ExecuteSearchOnChildNodes(const TArray<FStr
 	}
 }
 
+bool FConsoleVariablesEditorListRow::GetIsSelected() const
+{
+	return bIsSelected;
+}
+
+void FConsoleVariablesEditorListRow::SetIsSelected(const bool bNewSelected)
+{
+	bIsSelected = bNewSelected;
+}
+
 ECheckBoxState FConsoleVariablesEditorListRow::GetWidgetCheckedState() const
 {
 	return WidgetCheckedState;
@@ -184,7 +204,12 @@ bool FConsoleVariablesEditorListRow::IsRowChecked() const
 
 EVisibility FConsoleVariablesEditorListRow::GetDesiredVisibility() const
 {
-	return bDoesRowMatchSeachTerms || HasVisibleChildren() ? EVisibility::Visible : EVisibility::Collapsed;
+	return bDoesRowMatchSearchTerms || HasVisibleChildren() ? EVisibility::Visible : EVisibility::Collapsed;
+}
+
+TArray<FConsoleVariablesEditorListRowPtr> FConsoleVariablesEditorListRow::GetSelectedTreeViewItems() const
+{
+	return ListViewPtr.Pin()->GetSelectedTreeViewItems();
 }
 
 FReply FConsoleVariablesEditorListRow::OnRemoveButtonClicked()
diff --git a/Engine/Plugins/Experimental/ConsoleVariables/Source/ConsoleVariablesEditor/Private/Views/List/ConsoleVariablesEditorListRow.h b/Engine/Plugins/Experimental/ConsoleVariables/Source/ConsoleVariablesEditor/Private/Views/List/ConsoleVariablesEditorListRow.h
--- a/Engine/Plugins/Experimental/ConsoleVariables/Source/ConsoleVariablesEditor/Private/Views/List/ConsoleVariablesEditorListRow.h
+++ b/Engine/Plugins/Experimental/ConsoleVariables/Source/ConsoleVariablesEditor/Private/Views/List/ConsoleVariablesEditorListRow.h
@@ -25,25 +25,29 @@ struct FConsoleVariablesEditorListRow final : TSharedFromThis<FConsoleVariablesE
 
 	void FlushReferences();
 	
-	FConsoleVariablesEditorListRow(const TWeakPtr<FConsoleVariablesEditorCommandInfo> InCommandInfo, const FString& InPresetValue, const EConsoleVariablesEditorListRowType InRowType, 
+	FConsoleVariablesEditorListRow(
+		const TWeakPtr<FConsoleVariablesEditorCommandInfo> InCommandInfo, const FString& InPresetValue, const EConsoleVariablesEditorListRowType InRowType, 
 		const ECheckBoxState StartingWidgetCheckboxState, const TSharedRef<SConsoleVariablesEditorList>& InListView, 
-		const TWeakPtr<FConsoleVariablesEditorListRow>& InDirectParentRow)
+		const int32 IndexInList, const TWeakPtr<FConsoleVariablesEditorListRow>& InDirectParentRow)
 	: CommandInfo(InCommandInfo)
 	, PresetValue(InPresetValue)
 	, RowType(InRowType)
 	, WidgetCheckedState(StartingWidgetCheckboxState)
 	, ListViewPtr(InListView)
+	, SortOrder(IndexInList)
 	, DirectParentRow(InDirectParentRow)
-	{};
+	{}
 
 	[[nodiscard]] TWeakPtr<FConsoleVariablesEditorCommandInfo> GetCommandInfo() const;
 
 	[[nodiscard]] EConsoleVariablesEditorListRowType GetRowType() const;
 
 	[[nodiscard]] int32 GetChildDepth() const;
-
 	void SetChildDepth(const int32 InDepth);
 
+	[[nodiscard]] int32 GetSortOrder() const;
+	void SetSortOrder(const int32 InNewOrder);
+
 	TWeakPtr<FConsoleVariablesEditorListRow> GetDirectParentRow() const;
 	void SetDirectParentRow(const TWeakPtr<FConsoleVariablesEditorListRow>& InDirectParentRow);
 	
@@ -75,6 +79,9 @@ struct FConsoleVariablesEditorListRow final : TSharedFromThis<FConsoleVariablesE
 	void ExecuteSearchOnChildNodes(const FString& SearchString) const;
 	void ExecuteSearchOnChildNodes(const TArray<FString>& Tokens) const;
 
+	[[nodiscard]] bool GetIsSelected() const;
+	void SetIsSelected(const bool bNewSelected);
+
 	[[nodiscard]] ECheckBoxState GetWidgetCheckedState() const;
 	void SetWidgetCheckedState(const ECheckBoxState NewState, const bool bShouldUpdateHierarchyCheckedStates = false);
 
@@ -92,6 +99,8 @@ struct FConsoleVariablesEditorListRow final : TSharedFromThis<FConsoleVariablesE
 		return ListViewPtr;
 	}
 
+	[[nodiscard]] TArray<FConsoleVariablesEditorListRowPtr> GetSelectedTreeViewItems() const;
+
 	FReply OnRemoveButtonClicked();
 	
 	void ResetToPresetValue() const;
@@ -102,17 +111,21 @@ private:
 	FString PresetValue = "";
 	EConsoleVariablesEditorListRowType RowType = SingleCommand;
 	TArray<FConsoleVariablesEditorListRowPtr> ChildRows;
+
+	ECheckBoxState WidgetCheckedState = ECheckBoxState::Checked;
+	
+	TWeakPtr<SConsoleVariablesEditorList> ListViewPtr;
 	
 	bool bIsTreeViewItemExpanded = false;
 	bool bShouldFlashOnScrollIntoView = false;
 
 	int32 ChildDepth = 0;
 
-	bool bDoesRowMatchSeachTerms = true;
-	
-	ECheckBoxState WidgetCheckedState = ECheckBoxState::Checked;
+	int32 SortOrder = -1;
 
-	TWeakPtr<SConsoleVariablesEditorList> ListViewPtr;
+	bool bDoesRowMatchSearchTerms = true;
+
+	bool bIsSelected = false;
 	TWeakPtr<FConsoleVariablesEditorListRow> DirectParentRow;
 
 	// Used to expand all children on shift+click.
diff --git a/Engine/Plugins/Experimental/ConsoleVariables/Source/ConsoleVariablesEditor/Private/Views/List/SConsoleVariablesEditorList.cpp b/Engine/Plugins/Experimental/ConsoleVariables/Source/ConsoleVariablesEditor/Private/Views/List/SConsoleVariablesEditorList.cpp
--- a/Engine/Plugins/Experimental/ConsoleVariables/Source/ConsoleVariablesEditor/Private/Views/List/SConsoleVariablesEditorList.cpp
+++ b/Engine/Plugins/Experimental/ConsoleVariables/Source/ConsoleVariablesEditor/Private/Views/List/SConsoleVariablesEditorList.cpp
@@ -8,11 +8,13 @@
 
 #include "Framework/MultiBox/MultiBoxBuilder.h"
 #include "Widgets/Input/SCheckBox.h"
+#include "Widgets/Input/SComboBox.h"
 #include "Widgets/Input/SSearchBox.h"
 #include "Widgets/Layout/SWidgetSwitcher.h"
 
 #define LOCTEXT_NAMESPACE "ConsoleVariablesEditor"
 
+const FName SConsoleVariablesEditorList::CustomSortOrderColumnName(TEXT("Order"));
 const FName SConsoleVariablesEditorList::CheckBoxColumnName(TEXT("Column"));
 const FName SConsoleVariablesEditorList::VariableNameColumnName(TEXT("Name"));
 const FName SConsoleVariablesEditorList::ValueColumnName(TEXT("Value"));
@@ -36,11 +38,29 @@ void SConsoleVariablesEditorList::Construct(const FArguments& InArgs)
 			SNew(SHorizontalBox)
 
 			+SHorizontalBox::Slot()
+			.Padding(10.f, 1.f, 0.f, 1.f)
 			[
 				SAssignNew(ListSearchBoxPtr, SSearchBox)
 				.HintText(LOCTEXT("ConsoleVariablesEditorList_SearchHintText", "Search tracked variables, values, sources or help text..."))
 				.OnTextChanged_Raw(this, &SConsoleVariablesEditorList::OnListViewSearchTextChanged)
 			]
+
+			+SHorizontalBox::Slot()
+			.AutoWidth()
+			.Padding(10.f, 1.f, 15.f, 1.f)
+			.HAlign(HAlign_Right)
+			[
+				SAssignNew( ViewOptionsComboButton, SComboButton )
+				.ComboButtonStyle( FAppStyle::Get(), "SimpleComboButtonWithIcon" ) // Use the tool bar item style for this button
+				.OnGetMenuContent( this, &SConsoleVariablesEditorList::BuildShowOptionsMenu)
+				.HasDownArrow(false)
+				.ButtonContent()
+				[
+					SNew(SImage)
+					.ColorAndOpacity(FSlateColor::UseForeground())
+					.Image( FAppStyle::Get().GetBrush("Icons.Settings") )
+				]
+			]
 		]
 
 		+ SVerticalBox::Slot()
@@ -57,7 +77,14 @@ void SConsoleVariablesEditorList::Construct(const FArguments& InArgs)
 			[
 				SAssignNew(TreeViewPtr, STreeView<FConsoleVariablesEditorListRowPtr>)
 				.HeaderRow(HeaderRow)
-				.SelectionMode(ESelectionMode::None)
+				.SelectionMode(ESelectionMode::Multi)
+				.OnSelectionChanged_Lambda([this] (const FConsoleVariablesEditorListRowPtr& Row, const ESelectInfo::Type SelectionType)
+				{
+					if(Row.IsValid())
+					{
+						Row->SetIsSelected(TreeViewPtr->GetSelectedItems().Contains(Row));
+					}
+				})
 				.TreeItemsSource(&TreeViewRootObjects)
 				.OnGenerateRow_Lambda([this](FConsoleVariablesEditorListRowPtr Row, const TSharedRef<STableViewBase>& OwnerTable)
 					{
@@ -80,7 +107,7 @@ void SConsoleVariablesEditorList::Construct(const FArguments& InArgs)
 				.Text(LOCTEXT("ConsoleVariablesEditorList_NoList", "No List to show. Try clearing the active search or adding some console variables to the list."))
 			]
 		]
-	];
+	]; 
 }
 
 SConsoleVariablesEditorList::~SConsoleVariablesEditorList()
@@ -93,24 +120,58 @@ SConsoleVariablesEditorList::~SConsoleVariablesEditorList()
 	TreeViewPtr.Reset();
 }
 
-FMenuBuilder SConsoleVariablesEditorList::BuildShowOptionsMenu()
+TSharedRef<SWidget> SConsoleVariablesEditorList::BuildShowOptionsMenu()
 {
 	FMenuBuilder ShowOptionsMenuBuilder = FMenuBuilder(true, nullptr);
 
-	ShowOptionsMenuBuilder.AddMenuEntry(
-		LOCTEXT("CollapseAll", "Collapse All"),
-		LOCTEXT("ConsoleVariablesEditorList_CollapseAll_Tooltip", "Collapse all expanded actor groups in the Modified Actors list."),
-		FSlateIcon(),
-		FUIAction(
-			FExecuteAction::CreateLambda([this]() {
-				SetAllGroupsCollapsed();
-				})
-		),
-		NAME_None,
-		EUserInterfaceActionType::Button
-	);
+	ShowOptionsMenuBuilder.BeginSection("AssetThumbnails", LOCTEXT("ShowHeading", "Show"));
+	{
+		// Add mode filters
+		auto AddFiltersLambda = [&ShowOptionsMenuBuilder](const FText& FilterTitle, const FText& FilterTooltip)
+		{
+			ShowOptionsMenuBuilder.AddMenuEntry(
+			   FilterTitle,
+			   FilterTooltip,
+			   FSlateIcon(),
+			   FUIAction(
+				   //FExecuteAction::CreateRaw( this, &SConsoleVariablesEditorList::ToggleFilterActive ),
+				   //FCanExecuteAction(),
+				   //FIsActionChecked::CreateRaw( this, &SConsoleVariablesEditorList::IsFilterActive )
+			   ),
+			   NAME_None,
+			   EUserInterfaceActionType::ToggleButton
+		   );
+		};
+
+		//for (Filter : Filters) { AddFiltersLambda(Filter.Title, Filter.TooltipText); }
+	}
+	ShowOptionsMenuBuilder.EndSection();
+	
+	ShowOptionsMenuBuilder.BeginSection("AssetThumbnails", LOCTEXT("SortHeading", "Sort"));
+	{
+		// Add commands
+
+		/*ShowOptionsMenuBuilder.AddMenuEntry(
+			LOCTEXT("CollapseAll", "Collapse All"),
+			LOCTEXT("ConsoleVariablesEditorList_CollapseAll_Tooltip", "Collapse all expanded actor groups in the Modified Actors list."),
+			FSlateIcon(),
+			FUIAction(FExecuteAction::CreateRaw(this, &SConsoleVariablesEditorList::SetAllGroupsCollapsed)),
+			NAME_None,
+			EUserInterfaceActionType::Button
+		);*/
+		
+		ShowOptionsMenuBuilder.AddMenuEntry(
+			LOCTEXT("SetSortOrder", "Set Sort Order"),
+			LOCTEXT("ConsoleVariablesEditorList_SetSortOrder_Tooltip", "Makes the current order of the variables list the saved order."),
+			FSlateIcon(),
+			FUIAction(FExecuteAction::CreateRaw(this, &SConsoleVariablesEditorList::SetSortOrder)),
+			NAME_None,
+			EUserInterfaceActionType::Button
+		);
+	}
+	ShowOptionsMenuBuilder.EndSection();
 
-	return ShowOptionsMenuBuilder;
+	return ShowOptionsMenuBuilder.MakeWidget();
 }
 
 void SConsoleVariablesEditorList::FlushMemory(const bool bShouldKeepMemoryAllocated)
@@ -165,6 +226,23 @@ void SConsoleVariablesEditorList::RefreshList(const FString& InConsoleCommandToS
 	}
 }
 
+TArray<FConsoleVariablesEditorListRowPtr> SConsoleVariablesEditorList::GetSelectedTreeViewItems() const
+{
+	return TreeViewPtr->GetSelectedItems();
+}
+
+TArray<FConsoleVariablesEditorListRowPtr> SConsoleVariablesEditorList::GetTreeViewItems() const
+{
+	return TreeViewRootObjects;
+}
+
+void SConsoleVariablesEditorList::SetTreeViewItems(const TArray<FConsoleVariablesEditorListRowPtr>& InItems)
+{
+	TreeViewRootObjects = InItems;
+
+	TreeViewPtr->RequestListRefresh();
+}
+
 void SConsoleVariablesEditorList::UpdatePresetValuesForSave(TObjectPtr<UConsoleVariablesAsset> InAsset) const
 {
 	TMap<FString, FString> NewSavedValueMap;
@@ -221,7 +299,7 @@ void SConsoleVariablesEditorList::GenerateTreeView()
 				FConsoleVariablesEditorListRowPtr NewRow = 
 					MakeShared<FConsoleVariablesEditorListRow>(
 							CommandInfo.Pin(), CommandAndValue.Value, FConsoleVariablesEditorListRow::SingleCommand, 
-							ECheckBoxState::Checked, SharedThis(this), nullptr);
+							ECheckBoxState::Checked, SharedThis(this), TreeViewRootObjects.Num(), nullptr);
 				TreeViewRootObjects.Add(NewRow);
 			}
 		}
@@ -240,6 +318,16 @@ TSharedPtr<SHeaderRow> SConsoleVariablesEditorList::GenerateHeaderRow()
 {
 	check(HeaderRow);
 	HeaderRow->ClearColumns();
+
+	HeaderRow->AddColumn(
+		SHeaderRow::Column(CustomSortOrderColumnName)
+			.DefaultLabel(FText::FromString("#"))
+			.HAlignHeader(EHorizontalAlignment::HAlign_Center)
+			.FillWidth(0.3f)
+			.ShouldGenerateWidget(true)
+			.SortMode_Raw(this, &SConsoleVariablesEditorList::GetSortMode, CustomSortOrderColumnName)
+			.OnSort_Raw(this, &SConsoleVariablesEditorList::OnSortColumnCalled)
+	);
 	
 	HeaderRow->AddColumn(
 		SHeaderRow::Column(CheckBoxColumnName)
@@ -270,6 +358,7 @@ TSharedPtr<SHeaderRow> SConsoleVariablesEditorList::GenerateHeaderRow()
 		SHeaderRow::Column(VariableNameColumnName)
 			.DefaultLabel(LOCTEXT("ConsoleVariablesEditorList_ConsoleVariableNameHeaderText", "Console Variable Name"))
 			.HAlignHeader(EHorizontalAlignment::HAlign_Left)
+			.FillWidth(1.7f)
 			.ShouldGenerateWidget(true)
 			.SortMode_Raw(this, &SConsoleVariablesEditorList::GetSortMode, VariableNameColumnName)
 			.OnSort_Raw(this, &SConsoleVariablesEditorList::OnSortColumnCalled)
@@ -293,7 +382,7 @@ TSharedPtr<SHeaderRow> SConsoleVariablesEditorList::GenerateHeaderRow()
 	return HeaderRow;
 }
 
-FReply SConsoleVariablesEditorList::SetAllGroupsCollapsed()
+void SConsoleVariablesEditorList::SetAllGroupsCollapsed()
 {
 	if (TreeViewPtr.IsValid())
 	{
@@ -308,8 +397,17 @@ FReply SConsoleVariablesEditorList::SetAllGroupsCollapsed()
 			RootRow->SetIsTreeViewItemExpanded(false);
 		}
 	}
+}
+
+void SConsoleVariablesEditorList::SetSortOrder()
+{
+	for (int32 RowItr = 0; RowItr < TreeViewRootObjects.Num(); RowItr++)
+	{
+		const TSharedPtr<FConsoleVariablesEditorListRow>& ChildRow = TreeViewRootObjects[RowItr];
+		ChildRow->SetSortOrder(RowItr);
+	}
 
-	return FReply::Handled();
+	ExecuteSort(CustomSortOrderColumnName, CycleSortMode(CustomSortOrderColumnName));
 }
 
 void SConsoleVariablesEditorList::OnListViewSearchTextChanged(const FText& Text) const
@@ -461,35 +559,28 @@ EColumnSortMode::Type SConsoleVariablesEditorList::CycleSortMode(const FName& In
 		break;
 	}
 
-	SortingMap.Empty();
+	ClearSorting();
 	SortingMap.Add(InColumnName, ColumnSortMode);
 
 	return ColumnSortMode;
 }
 
 void SConsoleVariablesEditorList::ExecuteSort(const FName& InColumnName, const EColumnSortMode::Type InColumnSortMode)
-{
-	if (InColumnSortMode == EColumnSortMode::Ascending)
+{	
+	if (InColumnName.IsEqual(CustomSortOrderColumnName))
 	{
-		if (InColumnName.IsEqual(VariableNameColumnName))
-		{
-			TreeViewRootObjects.StableSort(SortByVariableNameAscending);
-		}
-		else if (InColumnName.IsEqual(SourceColumnName))
-		{
-			TreeViewRootObjects.StableSort(SortBySourceAscending);
-		}
+		TreeViewRootObjects.StableSort(
+			InColumnSortMode == EColumnSortMode::Ascending ? SortByOrderAscending : SortByOrderDescending);
 	}
-	else if (InColumnSortMode == EColumnSortMode::Descending)
+	if (InColumnName.IsEqual(SourceColumnName))
 	{
-		if (InColumnName.IsEqual(VariableNameColumnName))
-		{
-			TreeViewRootObjects.StableSort(SortByVariableNameDescending);
-		}
-		else if (InColumnName.IsEqual(SourceColumnName))
-		{
-			TreeViewRootObjects.StableSort(SortBySourceDescending);
-		}
+		TreeViewRootObjects.StableSort(
+			InColumnSortMode == EColumnSortMode::Ascending ? SortBySourceAscending : SortBySourceDescending);
+	}
+	if (InColumnName.IsEqual(VariableNameColumnName))
+	{
+		TreeViewRootObjects.StableSort(
+			InColumnSortMode == EColumnSortMode::Ascending ? SortByVariableNameAscending : SortByVariableNameDescending);
 	}
 	
 	TreeViewPtr->RequestTreeRefresh();
diff --git a/Engine/Plugins/Experimental/ConsoleVariables/Source/ConsoleVariablesEditor/Private/Views/List/SConsoleVariablesEditorList.h b/Engine/Plugins/Experimental/ConsoleVariables/Source/ConsoleVariablesEditor/Private/Views/List/SConsoleVariablesEditorList.h
--- a/Engine/Plugins/Experimental/ConsoleVariables/Source/ConsoleVariablesEditor/Private/Views/List/SConsoleVariablesEditorList.h
+++ b/Engine/Plugins/Experimental/ConsoleVariables/Source/ConsoleVariablesEditor/Private/Views/List/SConsoleVariablesEditorList.h
@@ -10,24 +10,10 @@
 
 class UConsoleVariablesAsset;
 class SBox;
+class SComboButton;
 class SSearchBox;
 class SHeaderRow;
 
-struct FConsoleVariablesEditorListSplitterManager;
-typedef TSharedPtr<FConsoleVariablesEditorListSplitterManager> FConsoleVariablesEditorListSplitterManagerPtr;
-
-struct FConsoleVariablesEditorListSplitterManager
-{
-	float NestedColumnWidth = 0.5f; // The right side of the first splitter which contains the nested splitter for the property widgets
-	float SnapshotPropertyColumnWidth = 0.5f;
-};
-
-enum class EConsoleVariablesEditorSortType
-{
-	SortByVariableName,
-	SortBySource
-};
-
 class SConsoleVariablesEditorList final : public SCompoundWidget
 {
 
@@ -42,7 +28,7 @@ public:
 
 	virtual ~SConsoleVariablesEditorList() override;
 
-	FMenuBuilder BuildShowOptionsMenu();
+	TSharedRef<SWidget> BuildShowOptionsMenu();
 	
 	void FlushMemory(const bool bShouldKeepMemoryAllocated);
 
@@ -50,6 +36,16 @@ public:
 
 	void RefreshList(const FString& InConsoleCommandToScrollTo = "");
 
+	[[nodiscard]] TArray<FConsoleVariablesEditorListRowPtr> GetSelectedTreeViewItems() const;
+
+	[[nodiscard]] TArray<FConsoleVariablesEditorListRowPtr> GetTreeViewItems() const;
+	void SetTreeViewItems(const TArray<FConsoleVariablesEditorListRowPtr>& InItems);
+
+	[[nodiscard]] int32 GetTreeViewItemCount() const
+	{
+		return TreeViewRootObjects.Num();
+	}
+
 	/** Updates the saved values in a UConsoleVariablesAsset so that the command/value map can be saved to disk */
 	void UpdatePresetValuesForSave(TObjectPtr<UConsoleVariablesAsset> InAsset) const;
 
@@ -74,9 +70,14 @@ public:
 	void OnSortColumnCalled(EColumnSortPriority::Type Priority, const FName& ColumnName, EColumnSortMode::Type SortMode);
 	EColumnSortMode::Type CycleSortMode(const FName& InColumnName);
 	void ExecuteSort(const FName& InColumnName, const EColumnSortMode::Type InColumnSortMode);
+	void ClearSorting()
+	{
+		SortingMap.Empty();
+	}
 
 	// Column Names
 
+	static const FName CustomSortOrderColumnName;
 	static const FName CheckBoxColumnName;
 	static const FName VariableNameColumnName;
 	static const FName ValueColumnName;
@@ -88,13 +89,15 @@ private:
 	TSharedPtr<SHeaderRow> GenerateHeaderRow();
 	ECheckBoxState HeaderCheckBoxState = ECheckBoxState::Checked;
 	
-	FReply SetAllGroupsCollapsed();
+	void SetAllGroupsCollapsed();
+	void SetSortOrder();
 
 	// Search
 	
 	void OnListViewSearchTextChanged(const FText& Text) const;
 
 	TSharedPtr<SSearchBox> ListSearchBoxPtr;
+	TSharedPtr<SComboButton> ViewOptionsComboButton;
 	TSharedPtr<SBox> ListBoxContainerPtr;
 
 	//  Tree View Implementation
@@ -113,28 +116,40 @@ private:
 	// Sorting
 	
 	TMap<FName, EColumnSortMode::Type> SortingMap;
-	
-	TFunctionRef<bool(const FConsoleVariablesEditorListRowPtr&, const FConsoleVariablesEditorListRowPtr&)> SortByVariableNameAscending =
+
+	TFunctionRef<bool(const FConsoleVariablesEditorListRowPtr&, const FConsoleVariablesEditorListRowPtr&)> SortByOrderAscending =
 		[](const FConsoleVariablesEditorListRowPtr& A, const FConsoleVariablesEditorListRowPtr& B)
-	{
-		return A->GetCommandInfo().Pin()->Command < B->GetCommandInfo().Pin()->Command;
-	};
+		{
+			return A->GetSortOrder() < B->GetSortOrder();
+		};
 
-	TFunctionRef<bool(const FConsoleVariablesEditorListRowPtr&, const FConsoleVariablesEditorListRowPtr&)> SortByVariableNameDescending =
+	TFunctionRef<bool(const FConsoleVariablesEditorListRowPtr&, const FConsoleVariablesEditorListRowPtr&)> SortByOrderDescending =
 		[](const FConsoleVariablesEditorListRowPtr& A, const FConsoleVariablesEditorListRowPtr& B)
-	{
-		return B->GetCommandInfo().Pin()->Command < A->GetCommandInfo().Pin()->Command;
-	};
+		{
+			return B->GetSortOrder() < A->GetSortOrder();
+		};
 
 	TFunctionRef<bool(const FConsoleVariablesEditorListRowPtr&, const FConsoleVariablesEditorListRowPtr&)> SortBySourceAscending =
 		[](const FConsoleVariablesEditorListRowPtr& A, const FConsoleVariablesEditorListRowPtr& B)
+		{
+			return A->GetCommandInfo().Pin()->GetSourceAsText().ToString() < B->GetCommandInfo().Pin()->GetSourceAsText().ToString();
+		};
+
+	TFunctionRef<bool(const FConsoleVariablesEditorListRowPtr&, const FConsoleVariablesEditorListRowPtr&)> SortBySourceDescending =
+		[](const FConsoleVariablesEditorListRowPtr& A, const FConsoleVariablesEditorListRowPtr& B)
+		{
+			return B->GetCommandInfo().Pin()->GetSourceAsText().ToString() < A->GetCommandInfo().Pin()->GetSourceAsText().ToString();
+		};
+	
+	TFunctionRef<bool(const FConsoleVariablesEditorListRowPtr&, const FConsoleVariablesEditorListRowPtr&)> SortByVariableNameAscending =
+		[](const FConsoleVariablesEditorListRowPtr& A, const FConsoleVariablesEditorListRowPtr& B)
 	{
-		return A->GetCommandInfo().Pin()->GetSourceAsText().ToString() < B->GetCommandInfo().Pin()->GetSourceAsText().ToString();
+		return A->GetCommandInfo().Pin()->Command < B->GetCommandInfo().Pin()->Command;
 	};
 
-	TFunctionRef<bool(const FConsoleVariablesEditorListRowPtr&, const FConsoleVariablesEditorListRowPtr&)> SortBySourceDescending =
+	TFunctionRef<bool(const FConsoleVariablesEditorListRowPtr&, const FConsoleVariablesEditorListRowPtr&)> SortByVariableNameDescending =
 		[](const FConsoleVariablesEditorListRowPtr& A, const FConsoleVariablesEditorListRowPtr& B)
 	{
-		return B->GetCommandInfo().Pin()->GetSourceAsText().ToString() < A->GetCommandInfo().Pin()->GetSourceAsText().ToString();
+		return B->GetCommandInfo().Pin()->Command < A->GetCommandInfo().Pin()->Command;
 	};
 };
diff --git a/Engine/Plugins/Experimental/ConsoleVariables/Source/ConsoleVariablesEditor/Private/Views/List/SConsoleVariablesEditorListRow.cpp b/Engine/Plugins/Experimental/ConsoleVariables/Source/ConsoleVariablesEditor/Private/Views/List/SConsoleVariablesEditorListRow.cpp
--- a/Engine/Plugins/Experimental/ConsoleVariables/Source/ConsoleVariablesEditor/Private/Views/List/SConsoleVariablesEditorListRow.cpp
+++ b/Engine/Plugins/Experimental/ConsoleVariables/Source/ConsoleVariablesEditor/Private/Views/List/SConsoleVariablesEditorListRow.cpp
@@ -6,6 +6,7 @@
 #include "ConsoleVariablesEditorStyle.h"
 #include "SConsoleVariablesEditorListValueInput.h"
 
+#include "Input/DragAndDrop.h"
 #include "Kismet/KismetMathLibrary.h"
 #include "Styling/AppStyle.h"
 #include "Styling/StyleColors.h"
@@ -17,6 +18,45 @@
 
 #define LOCTEXT_NAMESPACE "ConsoleVariablesEditor"
 
+class FConsoleVariablesListRowDragDropOp : public FDecoratedDragDropOp
+{
+public:
+	DRAG_DROP_OPERATOR_TYPE(FConsoleVariablesListRowDragDropOp, FDecoratedDragDropOp)
+
+	/** The item being dragged and dropped */
+	TArray<FConsoleVariablesEditorListRowPtr> DraggedItems;
+
+	/** Constructs a new drag/drop operation */
+	static TSharedRef<FConsoleVariablesListRowDragDropOp> New(const TArray<FConsoleVariablesEditorListRowPtr>& InItems)
+	{
+		check(InItems.Num() > 0);
+
+		TSharedRef<FConsoleVariablesListRowDragDropOp> Operation = MakeShareable(new FConsoleVariablesListRowDragDropOp());
+
+		Operation->DraggedItems = InItems;
+
+		Operation->DefaultHoverIcon = FAppStyle::Get().GetBrush("Graph.ConnectorFeedback.Error");
+
+		// Set the display text and the transaction name based on whether we're dragging a single or multiple widgets
+		if (InItems.Num() == 1)
+		{
+			Operation->DefaultHoverText = FText::FromString(InItems[0]->GetCommandInfo().Pin()->Command);
+		}
+		else
+		{
+			Operation->DefaultHoverText =
+				FText::Format(
+					SConsoleVariablesEditorListRow::MultiDragFormatText,
+					FText::AsNumber(Operation->DraggedItems.Num())
+				);
+		}
+
+		Operation->Construct();
+
+		return Operation;
+	}
+};
+
 void SConsoleVariablesEditorListRow::Construct(
 	const FArguments& InArgs, const TSharedRef<STableViewBase>& InOwnerTable, const TWeakPtr<FConsoleVariablesEditorListRow> InRow)
 {
@@ -30,7 +70,11 @@ void SConsoleVariablesEditorListRow::Construct(
 
 	SMultiColumnTableRow<FConsoleVariablesEditorListRowPtr>::Construct(
 		FSuperRowType::FArguments()
-		.Padding(1.0f),
+		.Padding(1.0f)
+		.OnCanAcceptDrop(this, &SConsoleVariablesEditorListRow::HandleCanAcceptDrop)
+		.OnAcceptDrop(this, &SConsoleVariablesEditorListRow::HandleAcceptDrop)
+		.OnDragDetected(this, &SConsoleVariablesEditorListRow::HandleDragDetected)
+		.OnDragLeave(this, &SConsoleVariablesEditorListRow::HandleDragLeave),
 		InOwnerTable
 	);
 
@@ -56,8 +100,10 @@ TSharedRef<SWidget> SConsoleVariablesEditorListRow::GenerateWidgetForColumn(cons
 	FlashImages.Add(FlashImage);
 
 	return SNew(SBox)
+			.Visibility(EVisibility::SelfHitTestInvisible)
 			[
 				SNew(SOverlay)
+				.Visibility(EVisibility::SelfHitTestInvisible)
 
 				+SOverlay::Slot()
 				[
@@ -81,22 +127,26 @@ TSharedRef<SWidget> SConsoleVariablesEditorListRow::GenerateWidgetForColumn(cons
 
 void SConsoleVariablesEditorListRow::OnMouseEnter(const FGeometry& MyGeometry, const FPointerEvent& MouseEvent)
 {
-	SCompoundWidget::OnMouseEnter(MyGeometry, MouseEvent);
+	bIsHovered = true;
 
 	if (HoverableWidgetsPtr.IsValid())
 	{
-		HoverableWidgetsPtr->SetVisibility(EVisibility::Visible);
+		HoverableWidgetsPtr->SetVisibility(EVisibility::SelfHitTestInvisible);
 	}
+	
+	SMultiColumnTableRow<FConsoleVariablesEditorListRowPtr>::OnMouseEnter(MyGeometry, MouseEvent);
 }
 
 void SConsoleVariablesEditorListRow::OnMouseLeave(const FPointerEvent& MouseEvent)
 {
-	SCompoundWidget::OnMouseLeave(MouseEvent);
-
+	bIsHovered = false;
+	
 	if (HoverableWidgetsPtr.IsValid())
 	{
 		HoverableWidgetsPtr->SetVisibility(EVisibility::Collapsed);
 	}
+	
+	SMultiColumnTableRow<FConsoleVariablesEditorListRowPtr>::OnMouseLeave(MouseEvent);
 }
 
 SConsoleVariablesEditorListRow::~SConsoleVariablesEditorListRow()
@@ -110,6 +160,122 @@ SConsoleVariablesEditorListRow::~SConsoleVariablesEditorListRow()
 	HoverableWidgetsPtr.Reset();
 }
 
+FReply SConsoleVariablesEditorListRow::HandleDragDetected(const FGeometry& MyGeometry, const FPointerEvent& MouseEvent)
+{
+	TArray<FConsoleVariablesEditorListRowPtr> DraggedItems = Item.Pin()->GetSelectedTreeViewItems();
+	TSharedRef<FConsoleVariablesListRowDragDropOp> Operation =
+		FConsoleVariablesListRowDragDropOp::New(DraggedItems);
+	
+	return FReply::Handled().BeginDragDrop(Operation);
+}
+
+void SConsoleVariablesEditorListRow::HandleDragLeave(const FDragDropEvent& DragDropEvent)
+{
+	if (TSharedPtr<FConsoleVariablesListRowDragDropOp> Operation =
+	DragDropEvent.GetOperationAs<FConsoleVariablesListRowDragDropOp>())
+	{
+		Operation->ResetToDefaultToolTip();
+	}
+}
+
+TOptional<EItemDropZone> SConsoleVariablesEditorListRow::HandleCanAcceptDrop(const FDragDropEvent& DragDropEvent,
+	EItemDropZone DropZone, FConsoleVariablesEditorListRowPtr TargetItem)
+{
+	TSharedPtr<FConsoleVariablesListRowDragDropOp> Operation =
+		DragDropEvent.GetOperationAs<FConsoleVariablesListRowDragDropOp>();
+
+	if (!Operation.IsValid())
+	{
+		return TOptional<EItemDropZone>();
+	}
+
+	const bool bIsDropDenied = !TargetItem.IsValid() || Operation->DraggedItems.Num() < 1 ||
+		!TargetItem->GetCommandInfo().IsValid() || Operation->DraggedItems.Contains(TargetItem);
+
+	FString BoolAsString = bIsDropDenied ? "true" : "false";
+
+	if (bIsDropDenied)
+	{
+		Operation->ResetToDefaultToolTip();
+		
+		return TOptional<EItemDropZone>();
+	}
+
+	FText ItemNameText = FText::FromString(Operation->DraggedItems[0]->GetCommandInfo().Pin()->Command);
+
+	if (Operation->DraggedItems.Num() > 1)
+	{
+		ItemNameText = FText::Format(MultiDragFormatText, FText::AsNumber(Operation->DraggedItems.Num()));
+	}
+
+	const FText DropPermittedText =
+		FText::Format(InsertFormatText,
+			ItemNameText,
+			DropZone == EItemDropZone::BelowItem ? BelowText : AboveText,
+			FText::FromString(TargetItem->GetCommandInfo().Pin()->Command)
+		);
+	
+	Operation->SetToolTip(
+		DropPermittedText,
+		FAppStyle::Get().GetBrush("Graph.ConnectorFeedback.OK")
+	);
+
+	// We have no behaviour yet for dropping one item onto another, so we'll treat it like we dropped it above
+	return DropZone == EItemDropZone::OntoItem ? EItemDropZone::AboveItem : DropZone;
+}
+
+FReply SConsoleVariablesEditorListRow::HandleAcceptDrop(const FDragDropEvent& DragDropEvent, EItemDropZone DropZone,
+	FConsoleVariablesEditorListRowPtr TargetItem)
+{
+	TSharedPtr<FConsoleVariablesListRowDragDropOp> Operation =
+		DragDropEvent.GetOperationAs<FConsoleVariablesListRowDragDropOp>();
+
+	if (!Operation.IsValid())
+	{
+		return FReply::Unhandled();
+	}
+
+	const TSharedPtr<SConsoleVariablesEditorList> ListView = Item.Pin()->GetListViewPtr().Pin();
+	
+	TArray<FConsoleVariablesEditorListRowPtr> DraggedItems = Operation->DraggedItems;
+
+	TArray<FConsoleVariablesEditorListRowPtr> AllTreeItemsCopy = ListView->GetTreeViewItems();
+	
+	for (FConsoleVariablesEditorListRowPtr DraggedItem : DraggedItems)
+	{
+		if (!DraggedItem.IsValid() || !AllTreeItemsCopy.Contains(DraggedItem))
+		{
+			continue;
+		}
+
+		AllTreeItemsCopy.Remove(DraggedItem);
+	}
+	
+	const int32 TargetIndex = AllTreeItemsCopy.IndexOfByKey(TargetItem);
+
+	if (TargetIndex > -1)
+	{
+		for (int32 ItemIndex = DraggedItems.Num() - 1; ItemIndex >= 0; ItemIndex--)
+		{
+			FConsoleVariablesEditorListRowPtr DraggedItem = DraggedItems[ItemIndex];
+			
+			if (!DraggedItem.IsValid() || AllTreeItemsCopy.Contains(DraggedItem))
+			{
+				continue;
+			}
+			
+			AllTreeItemsCopy.Insert(DraggedItem, DropZone == EItemDropZone::AboveItem ? TargetIndex : TargetIndex + 1);
+		}
+		
+		FScopedTransaction(LOCTEXT("DragDropOperationTransactionText", "Console Variables Drag & Drop Operation"));
+		ListView->SetTreeViewItems(AllTreeItemsCopy);
+
+		ListView->ClearSorting();
+	}
+
+	return FReply::Handled();
+}
+
 void SConsoleVariablesEditorListRow::FlashRow()
 {
 	FlashAnimation.Play(this->AsShared());
@@ -150,9 +316,21 @@ const FSlateBrush* SConsoleVariablesEditorListRow::GetBorderImage(
 
 TSharedRef<SWidget> SConsoleVariablesEditorListRow::GenerateCells(const FName& InColumnName, const TSharedPtr<FConsoleVariablesEditorListRow> PinnedItem)
 {
+	
+	if (InColumnName.IsEqual(SConsoleVariablesEditorList::CustomSortOrderColumnName))
+	{
+		return  SNew(STextBlock)
+				.Visibility(EVisibility::SelfHitTestInvisible)
+				.Justification(ETextJustify::Center)
+				.Text_Lambda([this, PinnedItem]()
+				{
+					return FText::AsNumber(PinnedItem->GetSortOrder() + 1);
+				});
+	}
 	if (InColumnName.IsEqual(SConsoleVariablesEditorList::CheckBoxColumnName))
 	{
 		return SNew(SBox)
+				.Visibility(EVisibility::SelfHitTestInvisible)
 				.HAlign(HAlign_Center)
 				[
 					SNew(SCheckBox)
@@ -163,6 +341,7 @@ TSharedRef<SWidget> SConsoleVariablesEditorListRow::GenerateCells(const FName& I
 	if (InColumnName.IsEqual(SConsoleVariablesEditorList::VariableNameColumnName))
 	{
 		return  SNew(STextBlock)
+				.Visibility(EVisibility::SelfHitTestInvisible)
 				.Text(FText::FromString(PinnedItem->GetCommandInfo().Pin()->Command));
 	}
 	if (InColumnName.IsEqual(SConsoleVariablesEditorList::ValueColumnName))
@@ -179,6 +358,7 @@ TSharedRef<SWidget> SConsoleVariablesEditorListRow::GenerateCells(const FName& I
 				.VAlign(VAlign_Center)
 				[
 					SNew(STextBlock)
+					.Visibility(EVisibility::SelfHitTestInvisible)
 					.Text_Lambda([this]()
 					{
 						return Item.Pin()->GetCommandInfo().Pin()->GetSourceAsText();
@@ -278,7 +458,7 @@ TSharedRef<SWidget> SConsoleVariablesEditorListRow::GenerateValueCellWidget(cons
 				if (PinnedItem.IsValid() && PinnedItem->GetRowType() == FConsoleVariablesEditorListRow::SingleCommand)
 				{
 					return PinnedItem->IsRowChecked() && PinnedItem->GetCommandInfo().Pin()->IsCurrentValueDifferentFromInputValue(PinnedItem->GetPresetValue()) ?
-						EVisibility::Visible : EVisibility::Hidden;
+						EVisibility::Visible : EVisibility::Collapsed;
 				}
 
 				return EVisibility::Collapsed;
@@ -319,6 +499,7 @@ void SConsoleVariablesEditorListRowHoverWidgets::Construct(const FArguments& InA
 		})
 		[
 			SNew(SImage)
+			.Visibility(EVisibility::SelfHitTestInvisible)
 			.Image(FAppStyle::Get().GetBrush("Icons.Delete"))
 			.ColorAndOpacity(FSlateColor::UseForeground())
 		]
diff --git a/Engine/Plugins/Experimental/ConsoleVariables/Source/ConsoleVariablesEditor/Private/Views/List/SConsoleVariablesEditorListRow.h b/Engine/Plugins/Experimental/ConsoleVariables/Source/ConsoleVariablesEditor/Private/Views/List/SConsoleVariablesEditorListRow.h
--- a/Engine/Plugins/Experimental/ConsoleVariables/Source/ConsoleVariablesEditor/Private/Views/List/SConsoleVariablesEditorListRow.h
+++ b/Engine/Plugins/Experimental/ConsoleVariables/Source/ConsoleVariablesEditor/Private/Views/List/SConsoleVariablesEditorListRow.h
@@ -6,6 +6,7 @@
 #include "SConsoleVariablesEditorList.h"
 
 #include "CoreMinimal.h"
+#include "DragAndDrop/DecoratedDragDropOp.h"
 #include "Widgets/Input/SButton.h"
 #include "Widgets/Input/SNumericEntryBox.h"
 #include "Widgets/Layout/SBorder.h"
@@ -27,12 +28,18 @@ public:
 	void Construct(const FArguments& InArgs, const TSharedRef<STableViewBase>& InOwnerTable, const TWeakPtr<FConsoleVariablesEditorListRow> InRow);
 
 	virtual TSharedRef<SWidget> GenerateWidgetForColumn(const FName& InColumnName) override;
-	
-	virtual void OnMouseEnter(const FGeometry& MyGeometry, const FPointerEvent& MouseEvent) override;
 
+	// Begin SWidget
+	virtual void OnMouseEnter(const FGeometry& MyGeometry, const FPointerEvent& MouseEvent) override;
 	virtual void OnMouseLeave(const FPointerEvent& MouseEvent) override;
+	// End SWidget
+	
+	virtual ~SConsoleVariablesEditorListRow() override;	
 
-	virtual ~SConsoleVariablesEditorListRow() override;
+	FReply HandleDragDetected(const FGeometry& MyGeometry, const FPointerEvent& MouseEvent);
+	void HandleDragLeave(const FDragDropEvent& DragDropEvent);
+	TOptional<EItemDropZone> HandleCanAcceptDrop(const FDragDropEvent& DragDropEvent, EItemDropZone DropZone, FConsoleVariablesEditorListRowPtr TargetItem);
+	FReply HandleAcceptDrop(const FDragDropEvent& DragDropEvent, EItemDropZone DropZone, FConsoleVariablesEditorListRowPtr TargetItem);
 	
 	void FlashRow();
 	
@@ -48,6 +55,18 @@ public:
 
 	TSharedRef<SWidget> GenerateValueCellWidget(const TSharedPtr<FConsoleVariablesEditorListRow> PinnedItem);
 
+	#define LOCTEXT_NAMESPACE "ConsoleVariablesEditor"
+
+	const FText ValueWidgetToolTipFormatText = LOCTEXT("ValueWidgetToolTipFormatText", "Custom Value: {0}\nPreset Value: {1}\nStartup Value: {2} (Set By {3})");
+	const FText RevertButtonFormatText = LOCTEXT("RevertButtonFormatText", "Reset to Preset Value: {0}");
+
+	const FText InsertFormatText = LOCTEXT("InsertAboveFormatText", "Insert {0} {1} {2}");
+	const FText AboveText = LOCTEXT("AboveListItem", "above");
+	const FText BelowText = LOCTEXT("BelowListItem", "below");
+	static const inline FText MultiDragFormatText = LOCTEXT("MultiDragFormatText", "{0} Items");
+
+	#undef LOCTEXT_NAMESPACE
+
 private:
 	
 	TWeakPtr<FConsoleVariablesEditorListRow> Item;
@@ -55,13 +74,6 @@ private:
 	TArray<TSharedPtr<SImage>> FlashImages;
 
 	TSharedPtr<SConsoleVariablesEditorListValueInput> ValueChildInputWidget;
-
-	#define LOCTEXT_NAMESPACE "ConsoleVariablesEditor"
-
-	FText ValueWidgetToolTipFormatText = LOCTEXT("ValueWidgetToolTipFormatText", "Custom Value: {0}\nPreset Value: {1}\nStartup Value: {2} (Set By {3})");
-	FText RevertButtonFormatText = LOCTEXT("RevertButtonFormatText", "Reset to Preset Value: {0}");
-
-	#undef LOCTEXT_NAMESPACE
 	
 	TSharedPtr<SConsoleVariablesEditorListRowHoverWidgets> HoverableWidgetsPtr;
 
@@ -69,6 +81,8 @@ private:
 
 	const float FlashAnimationDuration = 0.75f;
 	const FLinearColor FlashColor = FLinearColor::White;
+	
+	bool bIsHovered = false;
 };
 
 class SConsoleVariablesEditorListRowHoverWidgets : public SCompoundWidget
diff --git a/Engine/Plugins/Experimental/ConsoleVariables/Source/ConsoleVariablesEditor/Private/Views/List/SConsoleVariablesEditorListValueInput.cpp b/Engine/Plugins/Experimental/ConsoleVariables/Source/ConsoleVariablesEditor/Private/Views/List/SConsoleVariablesEditorListValueInput.cpp
--- a/Engine/Plugins/Experimental/ConsoleVariables/Source/ConsoleVariablesEditor/Private/Views/List/SConsoleVariablesEditorListValueInput.cpp
+++ b/Engine/Plugins/Experimental/ConsoleVariables/Source/ConsoleVariablesEditor/Private/Views/List/SConsoleVariablesEditorListValueInput.cpp
@@ -88,16 +88,16 @@ void SConsoleVariablesEditorListValueInput_Float::Construct(const FArguments& In
 		})
 		.OnValueChanged_Lambda([this] (const float InValue)
 		{
-			if (const TSharedPtr<FConsoleVariablesEditorListRow> PinnedItem = Item.Pin())
-			{
-				const FString ValueAsString = FString::SanitizeFloat(InValue);
-				
+			const FString ValueAsString = FString::SanitizeFloat(InValue);
+			
+			if (const TSharedPtr<FConsoleVariablesEditorListRow> PinnedItem = Item.Pin(); !GetCachedValue().Equals(ValueAsString))
+			{				
 				PinnedItem->GetCommandInfo().Pin()->ExecuteCommand(ValueAsString);
 
 				FConsoleVariablesEditorModule::Get().SendMultiUserConsoleVariableChange(PinnedItem->GetCommandInfo().Pin()->Command, ValueAsString);
 			}
 			
-			SetCachedValue(FString::SanitizeFloat(InValue));
+			SetCachedValue(ValueAsString);
 		})
 		.IsEnabled(this, &SConsoleVariablesEditorListValueInput::IsRowChecked)
 	];
@@ -143,16 +143,16 @@ void SConsoleVariablesEditorListValueInput_Int::Construct(const FArguments& InAr
 		})
 		.OnValueChanged_Lambda([this] (const int32 InValue)
 		{
-			if (const TSharedPtr<FConsoleVariablesEditorListRow> PinnedItem = Item.Pin())
-			{
-				const FString ValueAsString = FString::FromInt(InValue);
-				
+			const FString ValueAsString = FString::FromInt(InValue);
+			
+			if (const TSharedPtr<FConsoleVariablesEditorListRow> PinnedItem = Item.Pin(); !GetCachedValue().Equals(ValueAsString))
+			{				
 				PinnedItem->GetCommandInfo().Pin()->ExecuteCommand(ValueAsString);
 
 				FConsoleVariablesEditorModule::Get().SendMultiUserConsoleVariableChange(PinnedItem->GetCommandInfo().Pin()->Command, ValueAsString);
 			}
 
-			SetCachedValue(FString::FromInt(InValue));
+			SetCachedValue(ValueAsString);
 		})
 		.IsEnabled(this, &SConsoleVariablesEditorListValueInput::IsRowChecked)
 	];
@@ -198,16 +198,16 @@ void SConsoleVariablesEditorListValueInput_String::Construct(const FArguments& I
 		})
 		.OnTextCommitted_Lambda([this] (const FText& InValue, ETextCommit::Type InTextCommitType)
 		{
-			if (const TSharedPtr<FConsoleVariablesEditorListRow> PinnedItem = Item.Pin())
-			{
-				const FString ValueAsString = InValue.ToString();
-				
+			const FString ValueAsString = InValue.ToString();
+			
+			if (const TSharedPtr<FConsoleVariablesEditorListRow> PinnedItem = Item.Pin(); !GetCachedValue().Equals(ValueAsString))
+			{				
 				PinnedItem->GetCommandInfo().Pin()->ExecuteCommand(ValueAsString);
 
 				FConsoleVariablesEditorModule::Get().SendMultiUserConsoleVariableChange(PinnedItem->GetCommandInfo().Pin()->Command, ValueAsString);
 			}
 
-			SetCachedValue(InValue.ToString());
+			SetCachedValue(ValueAsString);
 		})
 		.IsEnabled(this, &SConsoleVariablesEditorListValueInput::IsRowChecked)
 	];
@@ -255,16 +255,16 @@ void SConsoleVariablesEditorListValueInput_Bool::Construct(const FArguments& InA
 		.MaxSliderValue(2)
 		.OnValueChanged_Lambda([this] (const int32 InValue)
 		{
-			if (const TSharedPtr<FConsoleVariablesEditorListRow> PinnedItem = Item.Pin())
-			{
-				const FString ValueAsString = FString::FromInt(InValue);
-				
+			const FString ValueAsString = FString::FromInt(InValue);
+			
+			if (const TSharedPtr<FConsoleVariablesEditorListRow> PinnedItem = Item.Pin(); !GetCachedValue().Equals(ValueAsString))
+			{				
 				PinnedItem->GetCommandInfo().Pin()->ExecuteCommand(ValueAsString);
 
 				FConsoleVariablesEditorModule::Get().SendMultiUserConsoleVariableChange(PinnedItem->GetCommandInfo().Pin()->Command, ValueAsString);
 			}
 
-			SetCachedValue(FString::FromInt(InValue));
+			SetCachedValue(ValueAsString);
 		})
 		.IsEnabled(this, &SConsoleVariablesEditorListValueInput::IsRowChecked)
 	];
