Console Variables Editor

-Added Custom variable input box

diff --git a/Engine/Plugins/Experimental/ConsoleVariables/Source/ConsoleVariablesEditor/Private/Views/MainPanel/SConsoleVariablesEditorCustomConsoleInputBox.cpp b/Engine/Plugins/Experimental/ConsoleVariables/Source/ConsoleVariablesEditor/Private/Views/MainPanel/SConsoleVariablesEditorCustomConsoleInputBox.cpp
new file mode 100644
--- /dev/null
+++ b/Engine/Plugins/Experimental/ConsoleVariables/Source/ConsoleVariablesEditor/Private/Views/MainPanel/SConsoleVariablesEditorCustomConsoleInputBox.cpp
@@ -0,0 +1,341 @@
+// Copyright Epic Games, Inc. All Rights Reserved.
+
+#include "Views/MainPanel/SConsoleVariablesEditorCustomConsoleInputBox.h"
+
+#include "ConsoleVariablesEditorLog.h"
+#include "ConsoleVariablesEditorModule.h"
+#include "Views/List/ConsoleVariablesEditorList.h"
+#include "Views/MainPanel/SConsoleVariablesEditorMainPanel.h"
+
+#include "Framework/Application/SlateApplication.h"
+#include "OutputLog/Public/OutputLogModule.h"
+#include "Styling/AppStyle.h"
+#include "Styling/StyleColors.h"
+#include "Widgets/Input/SEditableTextBox.h"
+#include "Widgets/Layout/SBorder.h"
+#include "Widgets/Layout/SBox.h"
+#include "Widgets/Text/STextBlock.h"
+#include "Widgets/Views/SListView.h"
+
+#define LOCTEXT_NAMESPACE "ConsoleVariablesEditor"
+
+void SConsoleVariablesEditorCustomConsoleInputBox::Construct(
+	const FArguments& InArgs, TWeakPtr<SConsoleVariablesEditorMainPanel> InMainPanelWidget)
+{
+	check(InMainPanelWidget.IsValid());
+
+	MainPanelWidget = InMainPanelWidget;
+	
+	ChildSlot
+	[
+		SAssignNew( SuggestionBox, SMenuAnchor )
+		.Method(GIsEditor ? EPopupMethod::CreateNewWindow : EPopupMethod::UseCurrentWindow)
+		.Placement( MenuPlacement_AboveAnchor )
+		[
+			SAssignNew(InputText, SEditableTextBox)
+			.Font(FEditorStyle::Get().GetWidgetStyle<FTextBlockStyle>("Log.Normal").Font)
+			.HintText(LOCTEXT("ConsoleCommandExecutorHintText", "Enter Console Command"))
+			.OnTextChanged(this, &SConsoleVariablesEditorCustomConsoleInputBox::OnTextChanged)
+			.OnKeyCharHandler(this, &SConsoleVariablesEditorCustomConsoleInputBox::OnKeyCharHandler)
+			.OnKeyDownHandler(this, &SConsoleVariablesEditorCustomConsoleInputBox::HandleConsoleInputTextCommitted)
+			.ClearKeyboardFocusOnCommit(false)
+		
+		]
+		.MenuContent
+		(
+			SNew(SBorder)
+			.BorderImage(FEditorStyle::GetBrush("Menu.Background"))
+			.Padding( FMargin(2) )
+			[
+				SNew(SBox)
+				.HeightOverride(250) // avoids flickering, ideally this would be adaptive to the content without flickering
+				.MinDesiredWidth(300)
+				.MaxDesiredWidth_Lambda([this] ()
+				{
+					// Limit the width of the suggestions list to the work area that this widget currently resides on
+					const FSlateRect WidgetRect(
+						GetCachedGeometry().GetAbsolutePosition(),
+						GetCachedGeometry().GetAbsolutePosition() + GetCachedGeometry().GetAbsoluteSize());
+					const FSlateRect WidgetWorkArea = FSlateApplication::Get().GetWorkArea(WidgetRect);
+					return FMath::Max(300.0f, WidgetWorkArea.GetSize().X - 12.0f);
+				})
+				[
+					SAssignNew(SuggestionListView, SListView< TSharedPtr<FString> >)
+					.ListItemsSource(&Suggestions.SuggestionsList)
+					.SelectionMode( ESelectionMode::Single )							// Ideally the mouse over would not highlight while keyboard controls the UI
+					.OnGenerateRow_Lambda([this] (TSharedPtr<FString> Text, const TSharedRef<STableViewBase>& OwnerTable)
+					{
+						check(Text.IsValid());
+
+						FString SanitizedText = *Text;
+						SanitizedText.ReplaceInline(TEXT("\r\n"), TEXT("\n"), ESearchCase::CaseSensitive);
+						SanitizedText.ReplaceInline(TEXT("\r"), TEXT(" "), ESearchCase::CaseSensitive);
+						SanitizedText.ReplaceInline(TEXT("\n"), TEXT(" "), ESearchCase::CaseSensitive);
+
+						return
+							SNew(STableRow< TSharedPtr<FString> >, OwnerTable)
+							[
+								SNew(STextBlock)
+								.Text(FText::FromString(SanitizedText))
+								.TextStyle(FEditorStyle::Get(), "Log.Normal")
+								.HighlightText(Suggestions.SuggestionsHighlight)
+								.ColorAndOpacity(FSlateColor::UseForeground())
+							];
+					})
+					.OnSelectionChanged_Lambda([this] (TSharedPtr<FString> NewValue, ESelectInfo::Type SelectInfo)
+					{
+						if(bIgnoreUIUpdate)
+						{
+							return;
+						}
+						
+						Suggestions.SelectedSuggestion = Suggestions.SuggestionsList.IndexOfByPredicate([&NewValue](const TSharedPtr<FString>& InSuggestion)
+						{
+							return InSuggestion == NewValue;
+						});
+
+						MarkActiveSuggestion();
+
+						// If the user selected this suggestion by clicking on it, then go ahead and close the suggestion
+						// box as they've chosen the suggestion they're interested in.
+						if( SelectInfo == ESelectInfo::OnMouseClick )
+						{
+							SuggestionBox->SetIsOpen( false );
+
+							// Jump the caret to the end of the newly auto-completed line. This makes it so that selecting
+							// an option doesn't leave the cursor in the middle of the suggestion (which makes it hard to 
+							// ctrl-back out, or to type "?" for help, etc.)
+							InputText->GoTo(ETextLocation::EndOfDocument);
+						}
+					})
+					.ItemHeight(18)
+				]
+			]
+		)
+	];
+}
+
+FReply SConsoleVariablesEditorCustomConsoleInputBox::OnPreviewKeyDown(const FGeometry& MyGeometry, const FKeyEvent& KeyEvent)
+{
+	if(SuggestionBox->IsOpen())
+	{
+		if(KeyEvent.GetKey() == EKeys::Up || KeyEvent.GetKey() == EKeys::Down)
+		{
+			Suggestions.StepSelectedSuggestion(KeyEvent.GetKey() == EKeys::Up ? -1 : +1);
+			MarkActiveSuggestion();
+
+			return FReply::Handled();
+		}
+		else if (KeyEvent.GetKey() == EKeys::Tab)
+		{
+			if (Suggestions.HasSuggestions())
+			{
+				if (Suggestions.HasSelectedSuggestion())
+				{
+					Suggestions.StepSelectedSuggestion(KeyEvent.IsShiftDown() ? -1 : +1);
+				}
+				else
+				{
+					Suggestions.SelectedSuggestion = 0;
+				}
+				MarkActiveSuggestion();
+			}
+			
+			return FReply::Handled();
+		}
+		else if (KeyEvent.GetKey() == EKeys::Escape)
+		{
+			SuggestionBox->SetIsOpen(false);
+			return FReply::Handled();
+		}
+	}
+	else
+	{
+		if(KeyEvent.GetKey() == EKeys::Up)
+		{
+			// If the command field isn't empty we need you to have pressed Control+Up to summon the history (to make sure you're not just using caret navigation)
+			const bool bShowHistory = InputText->GetText().IsEmpty() || KeyEvent.IsControlDown();
+			if (bShowHistory)
+			{
+				TArray<FString> History;
+				IConsoleManager::Get().GetConsoleHistory(TEXT(""), History);
+				
+				SetSuggestions(History, FText::GetEmpty());
+				
+				if(Suggestions.HasSuggestions())
+				{
+					Suggestions.StepSelectedSuggestion(-1);
+					MarkActiveSuggestion();
+				}
+			}
+
+			// Need to always handle this for single-line controls to avoid them invoking widget navigation
+			if (bShowHistory)
+			{
+				return FReply::Handled();
+			}
+		}
+	}
+
+	return FReply::Unhandled();
+}
+
+void SConsoleVariablesEditorCustomConsoleInputBox::OnTextChanged(const FText& InText)
+{
+	if(bIgnoreUIUpdate)
+	{
+		return;
+	}
+	
+	const FString& InputTextStr = InputText->GetText().ToString();
+	if(!InputTextStr.IsEmpty())
+	{
+		TArray<FString> AutoCompleteList;
+		
+		auto OnConsoleVariable = [&AutoCompleteList](const TCHAR *Name, IConsoleObject* CVar)
+		{
+#if (UE_BUILD_SHIPPING || UE_BUILD_TEST)
+			if (CVar->TestFlags(ECVF_Cheat))
+			{
+				return;
+			}
+#endif // (UE_BUILD_SHIPPING || UE_BUILD_TEST)
+			if (CVar->TestFlags(ECVF_Unregistered))
+			{
+				return;
+			}
+
+			AutoCompleteList.Add(Name);
+		};
+
+		IConsoleManager::Get().ForEachConsoleObjectThatContains(FConsoleObjectVisitor::CreateLambda(OnConsoleVariable), *InputTextStr);
+
+		AutoCompleteList.Sort([InputTextStr](const FString& A, const FString& B)
+		{ 
+			if (A.StartsWith(InputTextStr))
+			{
+				if (!B.StartsWith(InputTextStr))
+				{
+					return true;
+				}
+			}
+			else
+			{
+				if (B.StartsWith(InputTextStr))
+				{
+					return false;
+				}
+			}
+
+			return A < B;
+
+		});
+
+
+		SetSuggestions(AutoCompleteList, FText::FromString(InputTextStr));
+	}
+	else
+	{
+		ClearSuggestions();
+	}
+}
+
+FReply SConsoleVariablesEditorCustomConsoleInputBox::OnKeyCharHandler(const FGeometry& MyGeometry, const FCharacterEvent& InCharacterEvent)
+{
+	// Intercept tab if used for auto-complete
+	if (InCharacterEvent.GetCharacter() == '\t')
+	{
+		return FReply::Handled();
+	}
+
+	if (InCharacterEvent.GetModifierKeys().AnyModifiersDown() && InCharacterEvent.GetCharacter() == ' ')
+	{	
+		// Ignore space bar + a modifier key.  It should not type a space as this is used by other keyboard shortcuts
+		return FReply::Handled();
+	}
+
+	return FReply::Unhandled();
+}
+
+FReply SConsoleVariablesEditorCustomConsoleInputBox::HandleConsoleInputTextCommitted(const FGeometry& MyGeometry, const FKeyEvent& KeyPressed)
+{
+	if (KeyPressed.GetKey().GetFName() == TEXT("Enter"))
+	{
+		const FText CommittedText = InputText->GetText();
+		bIgnoreUIUpdate = true;
+		InputText->SetText(FText::GetEmpty());
+		SuggestionBox->SetIsOpen( false );
+		bIgnoreUIUpdate = false;
+		return MainPanelWidget.Pin()->ValidateConsoleInput(CommittedText);
+	}
+
+	return FReply::Unhandled();
+}
+
+void SConsoleVariablesEditorCustomConsoleInputBox::SetSuggestions(TArray<FString>& Elements, FText Highlight)
+{
+	FString SelectionText;
+	if (Suggestions.HasSelectedSuggestion())
+	{
+		SelectionText = *Suggestions.GetSelectedSuggestion();
+	}
+
+	Suggestions.Reset();
+	Suggestions.SuggestionsHighlight = Highlight;
+
+	for(int32 i = 0; i < Elements.Num(); ++i)
+	{
+		Suggestions.SuggestionsList.Add(MakeShared<FString>(Elements[i]));
+
+		if (Elements[i] == SelectionText)
+		{
+			Suggestions.SelectedSuggestion = i;
+		}
+	}
+	SuggestionListView->RequestListRefresh();
+
+	if(Suggestions.HasSuggestions())
+	{
+		// Ideally if the selection box is open the output window is not changing it's window title (flickers)
+		SuggestionBox->SetIsOpen(true, false);
+		if (Suggestions.HasSelectedSuggestion())
+		{
+			SuggestionListView->RequestScrollIntoView(Suggestions.GetSelectedSuggestion());
+		}
+		else
+		{
+			SuggestionListView->ScrollToTop();
+		}
+	}
+	else
+	{
+		SuggestionBox->SetIsOpen(false);
+	}
+}
+
+void SConsoleVariablesEditorCustomConsoleInputBox::MarkActiveSuggestion()
+{
+	bIgnoreUIUpdate = true;
+	if (Suggestions.HasSelectedSuggestion())
+	{
+		TSharedPtr<FString> SelectedSuggestion = Suggestions.GetSelectedSuggestion();
+
+		SuggestionListView->SetSelection(SelectedSuggestion);
+		SuggestionListView->RequestScrollIntoView(SelectedSuggestion);	// Ideally this would only scroll if outside of the view
+
+		InputText->SetText(FText::FromString(*SelectedSuggestion));
+	}
+	else
+	{
+		SuggestionListView->ClearSelection();
+	}
+	bIgnoreUIUpdate = false;
+}
+
+void SConsoleVariablesEditorCustomConsoleInputBox::ClearSuggestions()
+{
+	SuggestionBox->SetIsOpen(false);
+	Suggestions.Reset();
+}
+
+#undef LOCTEXT_NAMESPACE
diff --git a/Engine/Plugins/Experimental/ConsoleVariables/Source/ConsoleVariablesEditor/Private/Views/MainPanel/SConsoleVariablesEditorCustomConsoleInputBox.h b/Engine/Plugins/Experimental/ConsoleVariables/Source/ConsoleVariablesEditor/Private/Views/MainPanel/SConsoleVariablesEditorCustomConsoleInputBox.h
new file mode 100644
--- /dev/null
+++ b/Engine/Plugins/Experimental/ConsoleVariables/Source/ConsoleVariablesEditor/Private/Views/MainPanel/SConsoleVariablesEditorCustomConsoleInputBox.h
@@ -0,0 +1,108 @@
+// Copyright Epic Games, Inc. All Rights Reserved.
+
+#pragma once
+
+#include "Widgets/SCompoundWidget.h"
+#include "Widgets/Input/SMenuAnchor.h"
+
+class SConsoleVariablesEditorMainPanel;
+class SEditableTextBox;
+
+class SConsoleVariablesEditorCustomConsoleInputBox : public SCompoundWidget
+{
+public:
+
+	SLATE_BEGIN_ARGS(SConsoleVariablesEditorCustomConsoleInputBox)
+	{}
+
+	SLATE_END_ARGS()
+	
+	struct FSuggestions
+	{
+		FSuggestions()
+			: SelectedSuggestion(INDEX_NONE)
+		{
+		}
+
+		void Reset()
+		{
+			SelectedSuggestion = INDEX_NONE;
+			SuggestionsList.Reset();
+			SuggestionsHighlight = FText::GetEmpty();
+		}
+
+		bool HasSuggestions() const
+		{
+			return SuggestionsList.Num() > 0;
+		}
+
+		bool HasSelectedSuggestion() const
+		{
+			return SuggestionsList.IsValidIndex(SelectedSuggestion);
+		}
+
+		void StepSelectedSuggestion(const int32 Step)
+		{
+			SelectedSuggestion += Step;
+			if (SelectedSuggestion < 0)
+			{
+				SelectedSuggestion = SuggestionsList.Num() - 1;
+			}
+			else if (SelectedSuggestion >= SuggestionsList.Num())
+			{
+				SelectedSuggestion = 0;
+			}
+		}
+
+		TSharedPtr<FString> GetSelectedSuggestion() const
+		{
+			return SuggestionsList.IsValidIndex(SelectedSuggestion) ? SuggestionsList[SelectedSuggestion] : nullptr;
+		}
+
+		/** INDEX_NONE if not set, otherwise index into SuggestionsList */
+		int32 SelectedSuggestion;
+
+		/** All log messages stored in this widget for the list view */
+		TArray<TSharedPtr<FString>> SuggestionsList;
+
+		/** Highlight text to use for the suggestions list */
+		FText SuggestionsHighlight;
+	};
+
+	void Construct(const FArguments& InArgs, TWeakPtr<SConsoleVariablesEditorMainPanel> InMainPanelWidget);
+
+	//~ Begin SWidget Interface
+	virtual FReply OnPreviewKeyDown(const FGeometry& MyGeometry, const FKeyEvent& KeyEvent) override;
+	//~ End SWidget Interface
+
+	void OnTextChanged(const FText& InText);
+
+	FReply OnKeyCharHandler(const FGeometry& MyGeometry, const FCharacterEvent& InCharacterEvent);
+
+	FReply HandleConsoleInputTextCommitted(const FGeometry& MyGeometry, const FKeyEvent& KeyPressed);
+	
+	void SetSuggestions(TArray<FString>& Elements, FText Highlight);
+
+	void MarkActiveSuggestion();
+
+	void ClearSuggestions();
+
+private:
+
+	TWeakPtr<SConsoleVariablesEditorMainPanel> MainPanelWidget;
+	
+	/** A reference to the actual text box inside ConsoleInput */
+	TSharedPtr<SEditableTextBox> InputText;
+
+	/** history / auto completion elements */
+	TSharedPtr<SMenuAnchor> SuggestionBox;
+
+	/** The list view for showing all log messages. Should be replaced by a full text editor */
+	TSharedPtr<SListView<TSharedPtr<FString>>> SuggestionListView;
+
+	/** Active list of suggestions */
+	FSuggestions Suggestions;
+
+	/** to prevent recursive calls in UI callback */
+	bool bIgnoreUIUpdate = false;
+};
diff --git a/Engine/Plugins/Experimental/ConsoleVariables/Source/ConsoleVariablesEditor/Private/Views/MainPanel/SConsoleVariablesEditorMainPanel.cpp b/Engine/Plugins/Experimental/ConsoleVariables/Source/ConsoleVariablesEditor/Private/Views/MainPanel/SConsoleVariablesEditorMainPanel.cpp
--- a/Engine/Plugins/Experimental/ConsoleVariables/Source/ConsoleVariablesEditor/Private/Views/MainPanel/SConsoleVariablesEditorMainPanel.cpp
+++ b/Engine/Plugins/Experimental/ConsoleVariables/Source/ConsoleVariablesEditor/Private/Views/MainPanel/SConsoleVariablesEditorMainPanel.cpp
@@ -9,12 +9,12 @@
 #include "MultiUser/ConsoleVariableSyncData.h"
 #include "Views/List/ConsoleVariablesEditorList.h"
 #include "Views/MainPanel/ConsoleVariablesEditorMainPanel.h"
+#include "Views/MainPanel/SConsoleVariablesEditorCustomConsoleInputBox.h"
 
 #include "ContentBrowserModule.h"
 #include "Framework/Application/SlateApplication.h"
 #include "Framework/MultiBox/MultiBoxBuilder.h"
 #include "IContentBrowserSingleton.h"
-#include "OutputLog/Public/OutputLogModule.h"
 #include "Styling/AppStyle.h"
 #include "Styling/StyleColors.h"
 #include "Widgets/Images/SImage.h"
@@ -27,20 +27,12 @@
 
 #define LOCTEXT_NAMESPACE "ConsoleVariablesEditor"
 
-void SConsoleVariablesEditorMainPanel::Construct(const FArguments& InArgs, const TSharedRef<FConsoleVariablesEditorMainPanel>& InMainPanel)
+void SConsoleVariablesEditorMainPanel::Construct(
+	const FArguments& InArgs, const TSharedRef<FConsoleVariablesEditorMainPanel>& InMainPanel)
 {
 	check(InMainPanel->GetEditorList().IsValid());
 
 	MainPanel = InMainPanel;
-
-	const FOutputLogModule& OutputLogModule = FModuleManager::LoadModuleChecked< FOutputLogModule >(TEXT("OutputLog"));
-
-	ConsoleInput = OutputLogModule.MakeConsoleInputBox(
-		ConsoleInputEditableTextBox, FSimpleDelegate::CreateLambda([](){}), FSimpleDelegate::CreateLambda([](){}));
-
-	check(ConsoleInput.IsValid());
-
-	ConsoleInputEditableTextBox->SetOnKeyDownHandler(FOnKeyDown::CreateRaw(this, &SConsoleVariablesEditorMainPanel::HandleConsoleInputTextCommitted));
 	
 	ChildSlot
 	[
@@ -50,7 +42,7 @@ void SConsoleVariablesEditorMainPanel::Construct(const FArguments& InArgs, const
 		
 		+SSplitter::Slot().SizeRule(SSplitter::ESizeRule::SizeToContent)
 		[
-	        GeneratePanelToolbar(ConsoleInput.ToSharedRef())
+	        GeneratePanelToolbar()
 		]
 
 		+SSplitter::Slot()
@@ -73,32 +65,29 @@ SConsoleVariablesEditorMainPanel::~SConsoleVariablesEditorMainPanel()
 
 }
 
-FReply SConsoleVariablesEditorMainPanel::HandleConsoleInputTextCommitted(const FGeometry& MyGeometry, const FKeyEvent& KeyPressed)
+FReply SConsoleVariablesEditorMainPanel::ValidateConsoleInput(const FText& CommittedText)
 {
-	if (KeyPressed.GetKey().GetFName() == TEXT("Enter"))
+	FString CommandString = CommittedText.ToString().TrimStartAndEnd();
+	FString ValueString;
+	
+	if (CommandString.Contains(" "))
 	{
-		FString CommandString = ConsoleInputEditableTextBox->GetText().ToString();
-		FString ValueString;
-		
-		if (CommandString.Contains(" "))
-		{
-			CommandString.Split(TEXT(" "), &CommandString, &ValueString);
-		}
-
-		if (IConsoleVariable* AsVariable = IConsoleManager::Get().FindConsoleVariable(*CommandString))
-		{
-			MainPanel.Pin()->AddConsoleVariable(CommandString, ValueString.IsEmpty() ? AsVariable->GetString() : ValueString, true);
-		}
-		else if (CommandString.IsEmpty())
-		{
-			UE_LOG(LogConsoleVariablesEditor, Warning, TEXT("hs: Input is blank."), __FUNCTION__);
-		}
-		else
-		{
-			UE_LOG(LogConsoleVariablesEditor, Warning, TEXT("%hs: Input %s is not a recognized console command."), __FUNCTION__, *CommandString);
-		}
-
-		ConsoleInputEditableTextBox->SetText(FText::GetEmpty());
+		CommandString.Split(TEXT(" "), &CommandString, &ValueString);
+	}
+
+	if (IConsoleVariable* AsVariable = IConsoleManager::Get().FindConsoleVariable(*CommandString))
+	{
+		MainPanel.Pin()->AddConsoleVariable(
+			CommandString, ValueString.IsEmpty() ? AsVariable->GetString() : ValueString, true);
+	}
+	else if (CommandString.IsEmpty())
+	{
+		UE_LOG(LogConsoleVariablesEditor, Warning, TEXT("hs: Input is blank."), __FUNCTION__);
+	}
+	else
+	{
+		UE_LOG(LogConsoleVariablesEditor, Warning, TEXT("%hs: Input %s is not a recognized console command."),
+			__FUNCTION__, *CommandString);
 	}
 
 	return FReply::Handled();
@@ -135,7 +124,7 @@ void SConsoleVariablesEditorMainPanel::ToggleMultiUserDetails(ECheckBoxState Che
 	}
 }
 
-TSharedRef<SWidget> SConsoleVariablesEditorMainPanel::GeneratePanelToolbar(const TSharedRef<SWidget> InConsoleInputWidget)
+TSharedRef<SWidget> SConsoleVariablesEditorMainPanel::GeneratePanelToolbar()
 {
 	return SNew(SBorder)
 	        .Padding(0)
@@ -150,7 +139,7 @@ TSharedRef<SWidget> SConsoleVariablesEditorMainPanel::GeneratePanelToolbar(const
 				.VAlign(VAlign_Fill)
 				.Padding(2.f, 2.f)
 				[
-					InConsoleInputWidget
+					SNew(SConsoleVariablesEditorCustomConsoleInputBox, SharedThis(this))
 				]
 
 				// Presets Management Button
@@ -224,10 +213,13 @@ TSharedRef<SWidget> SConsoleVariablesEditorMainPanel::OnGeneratePresetsMenu()
 {
 	FMenuBuilder MenuBuilder(true, nullptr);
 
-	IContentBrowserSingleton& ContentBrowser = FModuleManager::LoadModuleChecked<FContentBrowserModule>("ContentBrowser").Get();
+	IContentBrowserSingleton& ContentBrowser =
+		FModuleManager::LoadModuleChecked<FContentBrowserModule>("ContentBrowser").Get();
 
 	const FText LoadedPresetName = MainPanel.Pin()->GetReferenceAssetOnDisk().IsValid() ?
-		FText::Format(LoadedPresetFormatText, FText::FromString(MainPanel.Pin()->GetReferenceAssetOnDisk()->GetName())) : NoLoadedPresetText;
+		FText::Format(
+			LoadedPresetFormatText,
+			FText::FromString(MainPanel.Pin()->GetReferenceAssetOnDisk()->GetName())) : NoLoadedPresetText;
 
 	MenuBuilder.AddMenuEntry(
 		LoadedPresetName,
@@ -271,7 +263,8 @@ TSharedRef<SWidget> SConsoleVariablesEditorMainPanel::OnGeneratePresetsMenu()
 		AssetPickerConfig.AssetShowWarningText = LOCTEXT("NoPresets_Warning", "No Presets Found");
 		AssetPickerConfig.Filter.ClassNames.Add(UConsoleVariablesAsset::StaticClass()->GetFName());
 		AssetPickerConfig.Filter.bRecursiveClasses = true;
-		AssetPickerConfig.OnAssetSelected = FOnAssetSelected::CreateRaw(MainPanel.Pin().Get(), &FConsoleVariablesEditorMainPanel::ImportPreset);
+		AssetPickerConfig.OnAssetSelected =
+			FOnAssetSelected::CreateRaw(MainPanel.Pin().Get(), &FConsoleVariablesEditorMainPanel::ImportPreset);
 	}
 
 	MenuBuilder.BeginSection(NAME_None, LOCTEXT("ImportPreset_MenuSection", "Import Preset"));
diff --git a/Engine/Plugins/Experimental/ConsoleVariables/Source/ConsoleVariablesEditor/Private/Views/MainPanel/SConsoleVariablesEditorMainPanel.h b/Engine/Plugins/Experimental/ConsoleVariables/Source/ConsoleVariablesEditor/Private/Views/MainPanel/SConsoleVariablesEditorMainPanel.h
--- a/Engine/Plugins/Experimental/ConsoleVariables/Source/ConsoleVariablesEditor/Private/Views/MainPanel/SConsoleVariablesEditorMainPanel.h
+++ b/Engine/Plugins/Experimental/ConsoleVariables/Source/ConsoleVariablesEditor/Private/Views/MainPanel/SConsoleVariablesEditorMainPanel.h
@@ -7,7 +7,6 @@
 
 class FConsoleVariablesEditorMainPanel;
 class SBox;
-class SMultiLineEditableTextBox;
 class SHorizontalBox;
 class SVerticalBox;
 class UConsoleVariablesAsset;
@@ -27,7 +26,7 @@ public:
 	virtual ~SConsoleVariablesEditorMainPanel() override;
 
 	/** A handler for when text is committed to the Console Input Box that appears when clicking the "Add Console Variable" button */
-	FReply HandleConsoleInputTextCommitted(const FGeometry& MyGeometry, const FKeyEvent& KeyPressed);
+	FReply ValidateConsoleInput(const FText& CommittedText);
 
 	void RefreshMultiUserDetails();
 	void ToggleMultiUserDetails(ECheckBoxState CheckState);
@@ -42,13 +41,8 @@ private:
 	/** A reference to the button which opens the plugin settings */
 	TSharedPtr<SCheckBox> ConcertButtonPtr;
 
-	/** A special text box widget that can search console variables as text is typed into it */
-	TSharedPtr<SWidget> ConsoleInput;
-	/** A reference to the actual text box inside ConsoleInput */
-	TSharedPtr<SMultiLineEditableTextBox> ConsoleInputEditableTextBox;
-
 	/** Creates the toolbar at the top of the MainPanel widget */
-	TSharedRef<SWidget> GeneratePanelToolbar(const TSharedRef<SWidget> InConsoleInputWidget);
+	TSharedRef<SWidget> GeneratePanelToolbar();
 	
 	void CreateConcertButtonIfNeeded();
 
