Console Variables Editor:

-Sorting is now done Spotify-style

#[fyi] Jason.Walter

diff --git a/Engine/Plugins/Experimental/ConsoleVariables/Source/ConsoleVariablesEditor/Private/Views/List/SConsoleVariablesEditorList.cpp b/Engine/Plugins/Experimental/ConsoleVariables/Source/ConsoleVariablesEditor/Private/Views/List/SConsoleVariablesEditorList.cpp
--- a/Engine/Plugins/Experimental/ConsoleVariables/Source/ConsoleVariablesEditor/Private/Views/List/SConsoleVariablesEditorList.cpp
+++ b/Engine/Plugins/Experimental/ConsoleVariables/Source/ConsoleVariablesEditor/Private/Views/List/SConsoleVariablesEditorList.cpp
@@ -22,6 +22,10 @@ const FName SConsoleVariablesEditorList::SourceColumnName(TEXT("Source"));
 
 void SConsoleVariablesEditorList::Construct(const FArguments& InArgs)
 {
+	// Set Default Sorting info
+	ActiveSortingColumnName = CustomSortOrderColumnName;
+	ActiveSortingType = EColumnSortMode::Ascending;
+	
 	HeaderRow = SNew(SHeaderRow)
 				.CanSelectGeneratedColumn(true).Visibility(EVisibility::Visible);
 	
@@ -114,9 +118,13 @@ SConsoleVariablesEditorList::~SConsoleVariablesEditorList()
 {	
 	ListSearchBoxPtr.Reset();
 	ListBoxContainerPtr.Reset();
+	ViewOptionsComboButton.Reset();
 
 	FlushMemory(false);
 
+	HeaderRow.Reset();
+
+	TreeViewRootObjects.Empty();
 	TreeViewPtr.Reset();
 }
 
@@ -196,14 +204,8 @@ void SConsoleVariablesEditorList::RefreshList(const FString& InConsoleCommandToS
 	GenerateTreeView();
 
 	// Enforce Sort
-	TArray<FName> MapKeys;
-	if (SortingMap.GetKeys(MapKeys))
-	{
-		const FName& Key = MapKeys[0];
-		const EColumnSortMode::Type Mode = *SortingMap.Find(Key);
-		
-		ExecuteSort(Key, Mode);
-	}
+	const FName& SortingName = GetActiveSortingColumnName();
+	ExecuteSort(SortingName, GetSortModeForColumn(SortingName));
 
 	if (!InConsoleCommandToScrollTo.IsEmpty())
 	{
@@ -322,10 +324,11 @@ TSharedPtr<SHeaderRow> SConsoleVariablesEditorList::GenerateHeaderRow()
 	HeaderRow->AddColumn(
 		SHeaderRow::Column(CustomSortOrderColumnName)
 			.DefaultLabel(FText::FromString("#"))
+			.ToolTipText(LOCTEXT("ClickToSort","Click to sort"))
 			.HAlignHeader(EHorizontalAlignment::HAlign_Center)
 			.FillWidth(0.3f)
 			.ShouldGenerateWidget(true)
-			.SortMode_Raw(this, &SConsoleVariablesEditorList::GetSortMode, CustomSortOrderColumnName)
+			.SortMode_Raw(this, &SConsoleVariablesEditorList::GetSortModeForColumn, CustomSortOrderColumnName)
 			.OnSort_Raw(this, &SConsoleVariablesEditorList::OnSortColumnCalled)
 	);
 	
@@ -357,10 +360,11 @@ TSharedPtr<SHeaderRow> SConsoleVariablesEditorList::GenerateHeaderRow()
 	HeaderRow->AddColumn(
 		SHeaderRow::Column(VariableNameColumnName)
 			.DefaultLabel(LOCTEXT("ConsoleVariablesEditorList_ConsoleVariableNameHeaderText", "Console Variable Name"))
+			.ToolTipText(LOCTEXT("ClickToSort","Click to sort"))
 			.HAlignHeader(EHorizontalAlignment::HAlign_Left)
 			.FillWidth(1.7f)
 			.ShouldGenerateWidget(true)
-			.SortMode_Raw(this, &SConsoleVariablesEditorList::GetSortMode, VariableNameColumnName)
+			.SortMode_Raw(this, &SConsoleVariablesEditorList::GetSortModeForColumn, VariableNameColumnName)
 			.OnSort_Raw(this, &SConsoleVariablesEditorList::OnSortColumnCalled)
 	);
 
@@ -374,8 +378,9 @@ TSharedPtr<SHeaderRow> SConsoleVariablesEditorList::GenerateHeaderRow()
 	HeaderRow->AddColumn(
 		SHeaderRow::Column(SourceColumnName)
 			.DefaultLabel(LOCTEXT("ConsoleVariablesEditorList_SourceHeaderText", "Source"))
+			.ToolTipText(LOCTEXT("ClickToSort","Click to sort"))
 			.HAlignHeader(EHorizontalAlignment::HAlign_Left)
-			.SortMode_Raw(this, &SConsoleVariablesEditorList::GetSortMode, SourceColumnName)
+			.SortMode_Raw(this, &SConsoleVariablesEditorList::GetSortModeForColumn, SourceColumnName)
 			.OnSort_Raw(this, &SConsoleVariablesEditorList::OnSortColumnCalled)
 	);
 
@@ -514,13 +519,13 @@ void SConsoleVariablesEditorList::OnListItemCheckBoxStateChange(const ECheckBoxS
 	}
 }
 
-EColumnSortMode::Type SConsoleVariablesEditorList::GetSortMode(FName InColumnName) const
+EColumnSortMode::Type SConsoleVariablesEditorList::GetSortModeForColumn(FName InColumnName) const
 {
 	EColumnSortMode::Type ColumnSortMode = EColumnSortMode::None;
 
-	if (const EColumnSortMode::Type* FoundSortMode = SortingMap.Find(InColumnName))
+	if (GetActiveSortingColumnName().IsEqual(InColumnName))
 	{
-		ColumnSortMode = *FoundSortMode;
+		ColumnSortMode = ActiveSortingType;
 	}
 
 	return ColumnSortMode;
@@ -533,44 +538,27 @@ void SConsoleVariablesEditorList::OnSortColumnCalled(EColumnSortPriority::Type P
 
 EColumnSortMode::Type SConsoleVariablesEditorList::CycleSortMode(const FName& InColumnName)
 {
-	EColumnSortMode::Type ColumnSortMode = EColumnSortMode::None;
-
-	if (const EColumnSortMode::Type* FoundSortMode = SortingMap.Find(InColumnName))
+	// Custom handler for Custom Sort Order mode
+	if (InColumnName.IsEqual(CustomSortOrderColumnName))
 	{
-		ColumnSortMode = *FoundSortMode;
+		ActiveSortingType = EColumnSortMode::Ascending;
 	}
-
-	switch (ColumnSortMode)
+	else
 	{
-	case EColumnSortMode::None:
-		ColumnSortMode = EColumnSortMode::Ascending;
-		break;
-
-	case EColumnSortMode::Ascending:
-		ColumnSortMode = EColumnSortMode::Descending;
-		break;
-
-	case EColumnSortMode::Descending:
-		ColumnSortMode = EColumnSortMode::None;
-		break;
-
-	default:
-		ColumnSortMode = EColumnSortMode::None;
-		break;
+		const EColumnSortMode::Type PreviousColumnSortMode = GetSortModeForColumn(InColumnName);
+		ActiveSortingType = PreviousColumnSortMode ==
+			EColumnSortMode::Ascending ? EColumnSortMode::Descending : EColumnSortMode::Ascending;
 	}
 
-	ClearSorting();
-	SortingMap.Add(InColumnName, ColumnSortMode);
-
-	return ColumnSortMode;
+	ActiveSortingColumnName = InColumnName;
+	return ActiveSortingType;
 }
 
 void SConsoleVariablesEditorList::ExecuteSort(const FName& InColumnName, const EColumnSortMode::Type InColumnSortMode)
 {	
 	if (InColumnName.IsEqual(CustomSortOrderColumnName))
 	{
-		TreeViewRootObjects.StableSort(
-			InColumnSortMode == EColumnSortMode::Ascending ? SortByOrderAscending : SortByOrderDescending);
+		TreeViewRootObjects.StableSort(SortByOrderAscending);
 	}
 	if (InColumnName.IsEqual(SourceColumnName))
 	{
diff --git a/Engine/Plugins/Experimental/ConsoleVariables/Source/ConsoleVariablesEditor/Private/Views/List/SConsoleVariablesEditorList.h b/Engine/Plugins/Experimental/ConsoleVariables/Source/ConsoleVariablesEditor/Private/Views/List/SConsoleVariablesEditorList.h
--- a/Engine/Plugins/Experimental/ConsoleVariables/Source/ConsoleVariablesEditor/Private/Views/List/SConsoleVariablesEditorList.h
+++ b/Engine/Plugins/Experimental/ConsoleVariables/Source/ConsoleVariablesEditor/Private/Views/List/SConsoleVariablesEditorList.h
@@ -65,15 +65,21 @@ public:
 	void OnListItemCheckBoxStateChange(const ECheckBoxState InNewState);
 
 	// Sorting
-	
-	EColumnSortMode::Type GetSortMode(FName InColumnName) const;
+
+	const FName& GetActiveSortingColumnName() const
+	{
+		return ActiveSortingColumnName;
+	}
+	EColumnSortMode::Type GetSortModeForColumn(FName InColumnName) const;
 	void OnSortColumnCalled(EColumnSortPriority::Type Priority, const FName& ColumnName, EColumnSortMode::Type SortMode);
 	EColumnSortMode::Type CycleSortMode(const FName& InColumnName);
 	void ExecuteSort(const FName& InColumnName, const EColumnSortMode::Type InColumnSortMode);
 	void ClearSorting()
 	{
-		SortingMap.Empty();
+		ActiveSortingColumnName = NAME_None;
+		ActiveSortingType = EColumnSortMode::None;
 	}
+	void SetSortOrder();
 
 	// Column Names
 
@@ -90,7 +96,6 @@ private:
 	ECheckBoxState HeaderCheckBoxState = ECheckBoxState::Checked;
 	
 	void SetAllGroupsCollapsed();
-	void SetSortOrder();
 
 	// Search
 	
@@ -114,8 +119,9 @@ private:
 	TArray<FConsoleVariablesEditorListRowPtr> TreeViewRootObjects;
 
 	// Sorting
-	
-	TMap<FName, EColumnSortMode::Type> SortingMap;
+
+	FName ActiveSortingColumnName = NAME_None;
+	EColumnSortMode::Type ActiveSortingType = EColumnSortMode::None;
 
 	TFunctionRef<bool(const FConsoleVariablesEditorListRowPtr&, const FConsoleVariablesEditorListRowPtr&)> SortByOrderAscending =
 		[](const FConsoleVariablesEditorListRowPtr& A, const FConsoleVariablesEditorListRowPtr& B)
@@ -123,12 +129,6 @@ private:
 			return A->GetSortOrder() < B->GetSortOrder();
 		};
 
-	TFunctionRef<bool(const FConsoleVariablesEditorListRowPtr&, const FConsoleVariablesEditorListRowPtr&)> SortByOrderDescending =
-		[](const FConsoleVariablesEditorListRowPtr& A, const FConsoleVariablesEditorListRowPtr& B)
-		{
-			return B->GetSortOrder() < A->GetSortOrder();
-		};
-
 	TFunctionRef<bool(const FConsoleVariablesEditorListRowPtr&, const FConsoleVariablesEditorListRowPtr&)> SortBySourceAscending =
 		[](const FConsoleVariablesEditorListRowPtr& A, const FConsoleVariablesEditorListRowPtr& B)
 		{
diff --git a/Engine/Plugins/Experimental/ConsoleVariables/Source/ConsoleVariablesEditor/Private/Views/List/SConsoleVariablesEditorListRow.cpp b/Engine/Plugins/Experimental/ConsoleVariables/Source/ConsoleVariablesEditor/Private/Views/List/SConsoleVariablesEditorListRow.cpp
--- a/Engine/Plugins/Experimental/ConsoleVariables/Source/ConsoleVariablesEditor/Private/Views/List/SConsoleVariablesEditorListRow.cpp
+++ b/Engine/Plugins/Experimental/ConsoleVariables/Source/ConsoleVariablesEditor/Private/Views/List/SConsoleVariablesEditorListRow.cpp
@@ -180,7 +180,7 @@ void SConsoleVariablesEditorListRow::HandleDragLeave(const FDragDropEvent& DragD
 
 TOptional<EItemDropZone> SConsoleVariablesEditorListRow::HandleCanAcceptDrop(const FDragDropEvent& DragDropEvent,
 	EItemDropZone DropZone, FConsoleVariablesEditorListRowPtr TargetItem)
-{
+{	
 	TSharedPtr<FConsoleVariablesListRowDragDropOp> Operation =
 		DragDropEvent.GetOperationAs<FConsoleVariablesListRowDragDropOp>();
 
@@ -188,6 +188,20 @@ TOptional<EItemDropZone> SConsoleVariablesEditorListRow::HandleCanAcceptDrop(con
 	{
 		return TOptional<EItemDropZone>();
 	}
+	
+	const bool bShouldEvaluateDrop = Item.IsValid() && Item.Pin()->GetListViewPtr().IsValid() &&
+		Item.Pin()->GetListViewPtr().Pin()->GetActiveSortingColumnName().IsEqual(
+			SConsoleVariablesEditorList::CustomSortOrderColumnName);
+
+	if (!bShouldEvaluateDrop)
+	{
+		Operation->SetToolTip(
+			LOCTEXT("SortByCustomOrderDrgDropWarning","Sort by custom order (\"#\") to drag & drop"),
+			FAppStyle::Get().GetBrush("Graph.ConnectorFeedback.Error")
+		);
+			
+		return TOptional<EItemDropZone>();
+	}
 
 	const bool bIsDropDenied = !TargetItem.IsValid() || Operation->DraggedItems.Num() < 1 ||
 		!TargetItem->GetCommandInfo().IsValid() || Operation->DraggedItems.Contains(TargetItem);
@@ -267,10 +281,8 @@ FReply SConsoleVariablesEditorListRow::HandleAcceptDrop(const FDragDropEvent& Dr
 			AllTreeItemsCopy.Insert(DraggedItem, DropZone == EItemDropZone::AboveItem ? TargetIndex : TargetIndex + 1);
 		}
 		
-		FScopedTransaction(LOCTEXT("DragDropOperationTransactionText", "Console Variables Drag & Drop Operation"));
 		ListView->SetTreeViewItems(AllTreeItemsCopy);
-
-		ListView->ClearSorting();
+		ListView->SetSortOrder();
 	}
 
 	return FReply::Handled();
diff --git a/Engine/Plugins/Experimental/ConsoleVariables/Source/ConsoleVariablesEditor/Private/Views/MainPanel/SConsoleVariablesEditorCustomConsoleInputBox.cpp b/Engine/Plugins/Experimental/ConsoleVariables/Source/ConsoleVariablesEditor/Private/Views/MainPanel/SConsoleVariablesEditorCustomConsoleInputBox.cpp
--- a/Engine/Plugins/Experimental/ConsoleVariables/Source/ConsoleVariablesEditor/Private/Views/MainPanel/SConsoleVariablesEditorCustomConsoleInputBox.cpp
+++ b/Engine/Plugins/Experimental/ConsoleVariables/Source/ConsoleVariablesEditor/Private/Views/MainPanel/SConsoleVariablesEditorCustomConsoleInputBox.cpp
@@ -15,7 +15,6 @@
 #include "Widgets/Layout/SBorder.h"
 #include "Widgets/Layout/SBox.h"
 #include "Widgets/Text/STextBlock.h"
-#include "Widgets/Views/SListView.h"
 
 #define LOCTEXT_NAMESPACE "ConsoleVariablesEditor"
 
diff --git a/Engine/Plugins/Experimental/ConsoleVariables/Source/ConsoleVariablesEditor/Private/Views/MainPanel/SConsoleVariablesEditorCustomConsoleInputBox.h b/Engine/Plugins/Experimental/ConsoleVariables/Source/ConsoleVariablesEditor/Private/Views/MainPanel/SConsoleVariablesEditorCustomConsoleInputBox.h
--- a/Engine/Plugins/Experimental/ConsoleVariables/Source/ConsoleVariablesEditor/Private/Views/MainPanel/SConsoleVariablesEditorCustomConsoleInputBox.h
+++ b/Engine/Plugins/Experimental/ConsoleVariables/Source/ConsoleVariablesEditor/Private/Views/MainPanel/SConsoleVariablesEditorCustomConsoleInputBox.h
@@ -4,6 +4,7 @@
 
 #include "Widgets/SCompoundWidget.h"
 #include "Widgets/Input/SMenuAnchor.h"
+#include "Widgets/Views/SListView.h"
 
 class SConsoleVariablesEditorMainPanel;
 class SEditableTextBox;
