Console Variables Editor:

-Fix for: Missing text in Multi-User settings

#[fyi] Jason.Walter

diff --git a/Engine/Plugins/Experimental/ConsoleVariables/Source/ConsoleVariablesEditor/Private/MultiUser/ConcertConsoleVariableSessionCustomization.cpp b/Engine/Plugins/Experimental/ConsoleVariables/Source/ConsoleVariablesEditor/Private/MultiUser/ConcertConsoleVariableSessionCustomization.cpp
--- a/Engine/Plugins/Experimental/ConsoleVariables/Source/ConsoleVariablesEditor/Private/MultiUser/ConcertConsoleVariableSessionCustomization.cpp
+++ b/Engine/Plugins/Experimental/ConsoleVariables/Source/ConsoleVariablesEditor/Private/MultiUser/ConcertConsoleVariableSessionCustomization.cpp
@@ -362,7 +362,7 @@ void FConcertConsoleVariableSessionCustomization::SettingChange(const FConcertCV
 
 void FConcertConsoleVariableSessionCustomization::CustomizeDetails(IDetailLayoutBuilder& DetailLayout)
 {
-	IDetailCategoryBuilder& Settings = DetailLayout.EditCategory(TEXT("Multi-user Console Vari  Settings"));
+	IDetailCategoryBuilder& Settings = DetailLayout.EditCategory(TEXT("Multi-user Console Variable Settings"));
 
 	TSharedPtr<IPropertyHandle> SettingsProperty = DetailLayout.GetProperty(
 		GET_MEMBER_NAME_CHECKED(UConcertCVarConfig, LocalSettings));
