Console Variables Editor:

-Fix for: Current Preset should not be selectable

#[fyi] Jason.Walter

diff --git a/Engine/Plugins/Experimental/ConsoleVariables/Source/ConsoleVariablesEditor/Private/Views/MainPanel/SConsoleVariablesEditorMainPanel.cpp b/Engine/Plugins/Experimental/ConsoleVariables/Source/ConsoleVariablesEditor/Private/Views/MainPanel/SConsoleVariablesEditorMainPanel.cpp
--- a/Engine/Plugins/Experimental/ConsoleVariables/Source/ConsoleVariablesEditor/Private/Views/MainPanel/SConsoleVariablesEditorMainPanel.cpp
+++ b/Engine/Plugins/Experimental/ConsoleVariables/Source/ConsoleVariablesEditor/Private/Views/MainPanel/SConsoleVariablesEditorMainPanel.cpp
@@ -221,11 +221,15 @@ TSharedRef<SWidget> SConsoleVariablesEditorMainPanel::OnGeneratePresetsMenu()
 			LoadedPresetFormatText,
 			FText::FromString(MainPanel.Pin()->GetReferenceAssetOnDisk()->GetName())) : NoLoadedPresetText;
 
+	FUIAction PresetNameAction = FUIAction();
+	PresetNameAction.CanExecuteAction = FCanExecuteAction::CreateLambda([]() { return false; });
 	MenuBuilder.AddMenuEntry(
 		LoadedPresetName,
 		LoadedPresetName,
 		FSlateIcon(),
-		FUIAction()
+		PresetNameAction,
+		NAME_None,
+		EUserInterfaceActionType::None
 	);
 
 	MenuBuilder.AddMenuSeparator();
