Console Variables Editor:

-Fix for: Update tooltip when hovering Remove icon

#[fyi] Jason.Walter

diff --git a/Engine/Plugins/Experimental/ConsoleVariables/Source/ConsoleVariablesEditor/Private/Views/List/SConsoleVariablesEditorListRow.cpp b/Engine/Plugins/Experimental/ConsoleVariables/Source/ConsoleVariablesEditor/Private/Views/List/SConsoleVariablesEditorListRow.cpp
--- a/Engine/Plugins/Experimental/ConsoleVariables/Source/ConsoleVariablesEditor/Private/Views/List/SConsoleVariablesEditorListRow.cpp
+++ b/Engine/Plugins/Experimental/ConsoleVariables/Source/ConsoleVariablesEditor/Private/Views/List/SConsoleVariablesEditorListRow.cpp
@@ -504,6 +504,7 @@ void SConsoleVariablesEditorListRowHoverWidgets::Construct(const FArguments& InA
 	[
 		// Remove Button
 		SAssignNew(RemoveButtonPtr, SButton)
+		.ToolTipText(LOCTEXT("RemoveCvarTooltip","Remove cvar from this list and reset its value to the startup value."))
 		.ButtonColorAndOpacity(FStyleColors::Transparent)
 		.OnClicked_Lambda([this]()
 		{
