Console Variables Editor:

-Implemented Show Filters for Source text
-Edited order of operations in OnConsoleVariableChanged
-Added new method RebuildList and separated out lighter elements of it into RefreshList.
-Centered NoList text and set it to wrap automatically

diff --git a/Engine/Plugins/Experimental/ConsoleVariables/Source/ConsoleVariablesEditor/Private/ConsoleVariablesEditorModule.cpp b/Engine/Plugins/Experimental/ConsoleVariables/Source/ConsoleVariablesEditor/Private/ConsoleVariablesEditorModule.cpp
--- a/Engine/Plugins/Experimental/ConsoleVariables/Source/ConsoleVariablesEditor/Private/ConsoleVariablesEditorModule.cpp
+++ b/Engine/Plugins/Experimental/ConsoleVariables/Source/ConsoleVariablesEditor/Private/ConsoleVariablesEditorModule.cpp
@@ -92,7 +92,7 @@ void FConsoleVariablesEditorModule::QueryAndBeginTrackingConsoleVariables()
 				const FDelegateHandle Handle =
 					AsVariable->OnChangedDelegate().AddRaw(this, &FConsoleVariablesEditorModule::OnConsoleVariableChanged);
 				const TSharedRef<FConsoleVariablesEditorCommandInfo> Info =
-					MakeShared<FConsoleVariablesEditorCommandInfo>(Key, AsVariable, AsVariable->GetString(), Handle);
+					MakeShared<FConsoleVariablesEditorCommandInfo>(Key, AsVariable, Handle);
 				Info->StartupSource = Info->GetSource();
 				ConsoleVariablesMasterReference.Add(Info);
 			}
@@ -162,7 +162,7 @@ void FConsoleVariablesEditorModule::RegisterMenuItem()
 	BrowserSpawnerEntry.SetGroup(WorkspaceMenu::GetMenuStructure().GetLevelEditorCategory());
 }
 
-bool FConsoleVariablesEditorModule::RegisterProjectSettings()
+void FConsoleVariablesEditorModule::RegisterProjectSettings()
 {
 	ISettingsModule& SettingsModule = FModuleManager::LoadModuleChecked<ISettingsModule>("Settings");
 	{
@@ -173,11 +173,7 @@ bool FConsoleVariablesEditorModule::RegisterProjectSettings()
 			NSLOCTEXT("ConsoleVariables", "ConsoleVariablesSettingsDescription",
 			          "Configure the Console Variables Editor user settings"),
 			GetMutableDefault<UConsoleVariablesEditorProjectSettings>());
-
-		return true;
 	}
-
-	return false;
 }
 
 void FConsoleVariablesEditorModule::OnConsoleVariableChanged(IConsoleVariable* ChangedVariable)
@@ -188,18 +184,34 @@ void FConsoleVariablesEditorModule::OnConsoleVariableChanged(IConsoleVariable* C
 		FindCommandInfoByConsoleVariableReference(ChangedVariable); CommandInfo.IsValid())
 	{
 		FString OutValue;
-		const FString& Key = CommandInfo.Pin()->Command;
-		if (GetMutableDefault<UConsoleVariablesEditorProjectSettings>()->bAddAllChangedConsoleVariablesToCurrentPreset &&
-			!EditingAsset->FindSavedValueByCommandString(Key, OutValue) &&
-			CommandInfo.Pin()->IsCurrentValueDifferentFromInputValue(CommandInfo.Pin()->StartupValueAsString))
+		const TSharedPtr<FConsoleVariablesEditorCommandInfo>& PinnedCommand = CommandInfo.Pin();
+		const FString& Key = PinnedCommand->Command;
+		
+		const bool bIsVariableCurrentlyTracked = EditingAsset->FindSavedValueByCommandString(Key, OutValue);
+		
+		if (!bIsVariableCurrentlyTracked)
 		{
-			EditingAsset->AddOrSetConsoleVariableSavedValue(Key, ChangedVariable->GetString());
+			// If not yet tracked and we want to track variable changes from outside the dialogue,
+			// Check if the changed value differs from the startup value before tracking it
+			if (GetMutableDefault<UConsoleVariablesEditorProjectSettings>()->bAddAllChangedConsoleVariablesToCurrentPreset
+				&& PinnedCommand->IsCurrentValueDifferentFromInputValue(PinnedCommand->StartupValueAsString))
+			{
+				EditingAsset->AddOrSetConsoleVariableSavedValue(Key, ChangedVariable->GetString());
+				if (MainPanel.IsValid())
+				{
+					MainPanel->RebuildList();
+				}
 
+				SendMultiUserConsoleVariableChange(Key, ChangedVariable->GetString());
+			}
+		}
+		else // If it's already being tracked, refreshed the list to update show filters and other possibly stale elements
+		{
 			if (MainPanel.IsValid())
 			{
 				MainPanel->RefreshList();
 			}
-
+			
 			SendMultiUserConsoleVariableChange(Key, ChangedVariable->GetString());
 		}
 	}
@@ -224,7 +236,7 @@ TSharedRef<SDockTab> FConsoleVariablesEditorModule::SpawnMainPanelTab(const FSpa
 {
 	const TSharedRef<SDockTab> DockTab = SNew(SDockTab).TabRole(ETabRole::NomadTab);
 	DockTab->SetContent(MainPanel->GetOrCreateWidget());
-	MainPanel->RefreshList();
+	MainPanel->RebuildList();
 			
 	return DockTab;
 }
diff --git a/Engine/Plugins/Experimental/ConsoleVariables/Source/ConsoleVariablesEditor/Private/ConsoleVariablesEditorModule.h b/Engine/Plugins/Experimental/ConsoleVariables/Source/ConsoleVariablesEditor/Private/ConsoleVariablesEditorModule.h
--- a/Engine/Plugins/Experimental/ConsoleVariables/Source/ConsoleVariablesEditor/Private/ConsoleVariablesEditorModule.h
+++ b/Engine/Plugins/Experimental/ConsoleVariables/Source/ConsoleVariablesEditor/Private/ConsoleVariablesEditorModule.h
@@ -49,7 +49,7 @@ private:
 	void OnFEngineLoopInitComplete();
 
 	void RegisterMenuItem();
-	bool RegisterProjectSettings();
+	void RegisterProjectSettings();
 
 	void OnConsoleVariableChanged(IConsoleVariable* ChangedVariable);
 
diff --git a/Engine/Plugins/Experimental/ConsoleVariables/Source/ConsoleVariablesEditor/Private/Views/List/ConsoleVariablesEditorList.cpp b/Engine/Plugins/Experimental/ConsoleVariables/Source/ConsoleVariablesEditor/Private/Views/List/ConsoleVariablesEditorList.cpp
--- a/Engine/Plugins/Experimental/ConsoleVariables/Source/ConsoleVariablesEditor/Private/Views/List/ConsoleVariablesEditorList.cpp
+++ b/Engine/Plugins/Experimental/ConsoleVariables/Source/ConsoleVariablesEditor/Private/Views/List/ConsoleVariablesEditorList.cpp
@@ -19,11 +19,19 @@ TSharedRef<SWidget> FConsoleVariablesEditorList::GetOrCreateWidget()
 	return ListWidget.ToSharedRef();
 }
 
-void FConsoleVariablesEditorList::RefreshList(const FString& InConsoleCommandToScrollTo) const
+void FConsoleVariablesEditorList::RebuildList(const FString& InConsoleCommandToScrollTo) const
 {
 	if (ListWidget.IsValid())
 	{
-		ListWidget->RefreshList(InConsoleCommandToScrollTo);
+		ListWidget->RebuildList(InConsoleCommandToScrollTo);
+	}
+}
+
+void FConsoleVariablesEditorList::RefreshList() const
+{
+	if (ListWidget.IsValid())
+	{
+		ListWidget->RefreshList();
 	}
 }
 
diff --git a/Engine/Plugins/Experimental/ConsoleVariables/Source/ConsoleVariablesEditor/Private/Views/List/ConsoleVariablesEditorList.h b/Engine/Plugins/Experimental/ConsoleVariables/Source/ConsoleVariablesEditor/Private/Views/List/ConsoleVariablesEditorList.h
--- a/Engine/Plugins/Experimental/ConsoleVariables/Source/ConsoleVariablesEditor/Private/Views/List/ConsoleVariablesEditorList.h
+++ b/Engine/Plugins/Experimental/ConsoleVariables/Source/ConsoleVariablesEditor/Private/Views/List/ConsoleVariablesEditorList.h
@@ -20,7 +20,14 @@ public:
 
 	TSharedRef<SWidget> GetOrCreateWidget();
 
-	void RefreshList(const FString& InConsoleCommandToScrollTo = "") const;
+	/** Regenerate the list items and refresh the list. Call when adding or removing variables. */
+	void RebuildList(const FString& InConsoleCommandToScrollTo = "") const;
+
+	/**
+	 * Refresh filters and sorting.
+	 * Useful for when the list state has gone stale but the variable count has not changed.
+	 */
+	void RefreshList() const;
 
 	/** Updates the saved values in a UConsoleVariablesAsset so that the command/value map can be saved to disk */
 	void UpdatePresetValuesForSave(TObjectPtr<UConsoleVariablesAsset> InAsset) const;
diff --git a/Engine/Plugins/Experimental/ConsoleVariables/Source/ConsoleVariablesEditor/Private/Views/List/ConsoleVariablesEditorListFilters/ConsoleVariablesEditorListFilter_SourceText.h b/Engine/Plugins/Experimental/ConsoleVariables/Source/ConsoleVariablesEditor/Private/Views/List/ConsoleVariablesEditorListFilters/ConsoleVariablesEditorListFilter_SourceText.h
new file mode 100644
--- /dev/null
+++ b/Engine/Plugins/Experimental/ConsoleVariables/Source/ConsoleVariablesEditor/Private/Views/List/ConsoleVariablesEditorListFilters/ConsoleVariablesEditorListFilter_SourceText.h
@@ -0,0 +1,50 @@
+// Copyright Epic Games, Inc. All Rights Reserved.
+
+#pragma once
+
+#include "IConsoleVariablesEditorListFilter.h"
+
+#define LOCTEXT_NAMESPACE "ConsoleVariablesEditor"
+
+class FConsoleVariablesEditorListFilter_SourceText : public IConsoleVariablesEditorListFilter
+{
+public:
+
+	FConsoleVariablesEditorListFilter_SourceText(const FString& InFilterString)
+	: FilterString(InFilterString){}
+
+	virtual FString GetFilterName() override
+	{
+		return FilterString;
+	}
+
+	virtual FText GetFilterButtonLabel() override
+	{
+		return FText::Format(LOCTEXT("ShowSourceTextFilterFormat", "Show {0}"), FText::FromString(GetFilterName()));
+	}
+
+	virtual FText GetFilterButtonToolTip() override
+	{
+		return FText::Format(
+			LOCTEXT("ShowSourceTextFilterTooltipFormat", "Show rows that have a Source field matching '{0}'"),
+			FText::FromString(GetFilterName()));
+	}
+
+	virtual bool DoesItemPassFilter(const FConsoleVariablesEditorListRowPtr& InItem) override
+	{
+		if (InItem.IsValid(); const TSharedPtr<FConsoleVariablesEditorCommandInfo> PinnedCommand = InItem->GetCommandInfo().Pin();)
+		{
+			const FString& StringToSearch = PinnedCommand->GetSourceAsText().ToString();
+
+			return StringToSearch.Contains(FilterString);
+		}
+
+		return false;
+	}
+	
+private:
+
+	FString FilterString;
+};
+
+#undef LOCTEXT_NAMESPACE
diff --git a/Engine/Plugins/Experimental/ConsoleVariables/Source/ConsoleVariablesEditor/Private/Views/List/ConsoleVariablesEditorListFilters/IConsoleVariablesEditorListFilter.h b/Engine/Plugins/Experimental/ConsoleVariables/Source/ConsoleVariablesEditor/Private/Views/List/ConsoleVariablesEditorListFilters/IConsoleVariablesEditorListFilter.h
new file mode 100644
--- /dev/null
+++ b/Engine/Plugins/Experimental/ConsoleVariables/Source/ConsoleVariablesEditor/Private/Views/List/ConsoleVariablesEditorListFilters/IConsoleVariablesEditorListFilter.h
@@ -0,0 +1,44 @@
+// Copyright Epic Games, Inc. All Rights Reserved.
+
+#pragma once
+
+#include "../ConsoleVariablesEditorListRow.h"
+
+class IConsoleVariablesEditorListFilter
+{
+public:
+	
+	virtual ~IConsoleVariablesEditorListFilter() = default;
+
+	virtual FString GetFilterName() = 0;
+
+	/** Returns localized text to display for when this filter' active state is defined by a toggle button. */
+	virtual FText GetFilterButtonLabel() = 0;
+	/** Returns localized text to display when the user mouses over a toggle button that defines whether this filter is active. */
+	virtual FText GetFilterButtonToolTip() = 0;
+	
+	void SetFilterActive(const bool bNewEnabled)
+	{
+		bIsFilterActive = bNewEnabled;
+	}
+
+	void ToggleFilterActive()
+	{
+		bIsFilterActive = !bIsFilterActive;
+	}
+
+	bool GetIsFilterActive() const
+	{
+		return bIsFilterActive;
+	}
+
+	virtual bool DoesItemPassFilter(const FConsoleVariablesEditorListRowPtr& InItem)
+	{
+		return false;
+	}
+
+private:
+
+	bool bIsFilterActive = true;
+			
+};
diff --git a/Engine/Plugins/Experimental/ConsoleVariables/Source/ConsoleVariablesEditor/Private/Views/List/ConsoleVariablesEditorListRow.cpp b/Engine/Plugins/Experimental/ConsoleVariables/Source/ConsoleVariablesEditor/Private/Views/List/ConsoleVariablesEditorListRow.cpp
--- a/Engine/Plugins/Experimental/ConsoleVariables/Source/ConsoleVariablesEditor/Private/Views/List/ConsoleVariablesEditorListRow.cpp
+++ b/Engine/Plugins/Experimental/ConsoleVariables/Source/ConsoleVariablesEditor/Private/Views/List/ConsoleVariablesEditorListRow.cpp
@@ -177,6 +177,16 @@ void FConsoleVariablesEditorListRow::ExecuteSearchOnChildNodes(const TArray<FStr
 	}
 }
 
+bool FConsoleVariablesEditorListRow::GetDoesRowPassFilters() const
+{
+	return bDoesRowPassFilters;
+}
+
+void FConsoleVariablesEditorListRow::SetDoesRowPassFilters(const bool bPass)
+{
+	bDoesRowPassFilters = bPass;
+}
+
 bool FConsoleVariablesEditorListRow::GetIsSelected() const
 {
 	return bIsSelected;
@@ -204,7 +214,8 @@ bool FConsoleVariablesEditorListRow::IsRowChecked() const
 
 EVisibility FConsoleVariablesEditorListRow::GetDesiredVisibility() const
 {
-	return bDoesRowMatchSearchTerms || HasVisibleChildren() ? EVisibility::Visible : EVisibility::Collapsed;
+	return (bDoesRowMatchSearchTerms && bDoesRowPassFilters) || HasVisibleChildren() ?
+		EVisibility::Visible : EVisibility::Collapsed;
 }
 
 TArray<FConsoleVariablesEditorListRowPtr> FConsoleVariablesEditorListRow::GetSelectedTreeViewItems() const
@@ -228,7 +239,7 @@ FReply FConsoleVariablesEditorListRow::OnRemoveButtonClicked()
 
 	EditableAsset->RemoveConsoleVariable(CommandInfo.Pin()->Command);
 
-	ListViewPtr.Pin()->RefreshList();
+	ListViewPtr.Pin()->RebuildList();
 
 	return FReply::Handled();
 }
diff --git a/Engine/Plugins/Experimental/ConsoleVariables/Source/ConsoleVariablesEditor/Private/Views/List/ConsoleVariablesEditorListRow.h b/Engine/Plugins/Experimental/ConsoleVariables/Source/ConsoleVariablesEditor/Private/Views/List/ConsoleVariablesEditorListRow.h
--- a/Engine/Plugins/Experimental/ConsoleVariables/Source/ConsoleVariablesEditor/Private/Views/List/ConsoleVariablesEditorListRow.h
+++ b/Engine/Plugins/Experimental/ConsoleVariables/Source/ConsoleVariablesEditor/Private/Views/List/ConsoleVariablesEditorListRow.h
@@ -79,6 +79,9 @@ struct FConsoleVariablesEditorListRow final : TSharedFromThis<FConsoleVariablesE
 	void ExecuteSearchOnChildNodes(const FString& SearchString) const;
 	void ExecuteSearchOnChildNodes(const TArray<FString>& Tokens) const;
 
+	[[nodiscard]] bool GetDoesRowPassFilters() const;
+	void SetDoesRowPassFilters(const bool bPass);
+
 	[[nodiscard]] bool GetIsSelected() const;
 	void SetIsSelected(const bool bNewSelected);
 
@@ -124,6 +127,7 @@ private:
 	int32 SortOrder = -1;
 
 	bool bDoesRowMatchSearchTerms = true;
+	bool bDoesRowPassFilters = true;
 
 	bool bIsSelected = false;
 	TWeakPtr<FConsoleVariablesEditorListRow> DirectParentRow;
diff --git a/Engine/Plugins/Experimental/ConsoleVariables/Source/ConsoleVariablesEditor/Private/Views/List/SConsoleVariablesEditorList.cpp b/Engine/Plugins/Experimental/ConsoleVariables/Source/ConsoleVariablesEditor/Private/Views/List/SConsoleVariablesEditorList.cpp
--- a/Engine/Plugins/Experimental/ConsoleVariables/Source/ConsoleVariablesEditor/Private/Views/List/SConsoleVariablesEditorList.cpp
+++ b/Engine/Plugins/Experimental/ConsoleVariables/Source/ConsoleVariablesEditor/Private/Views/List/SConsoleVariablesEditorList.cpp
@@ -3,9 +3,11 @@
 #include "SConsoleVariablesEditorList.h"
 
 #include "ConsoleVariablesAsset.h"
+#include "ConsoleVariablesEditorListFilters/ConsoleVariablesEditorListFilter_SourceText.h"
 #include "ConsoleVariablesEditorModule.h"
 #include "SConsoleVariablesEditorListRow.h"
 
+#include "Algo/Find.h"
 #include "Framework/MultiBox/MultiBoxBuilder.h"
 #include "Widgets/Input/SCheckBox.h"
 #include "Widgets/Input/SComboBox.h"
@@ -30,6 +32,8 @@ void SConsoleVariablesEditorList::Construct(const FArguments& InArgs)
 				.CanSelectGeneratedColumn(true).Visibility(EVisibility::Visible);
 	
 	GenerateHeaderRow();
+
+	SetupFilters();
 	
 	ChildSlot
 	[
@@ -104,14 +108,18 @@ void SConsoleVariablesEditorList::Construct(const FArguments& InArgs)
 
 			// For when no rows exist in view
 			+ SWidgetSwitcher::Slot()
-			.HAlign(HAlign_Center)
+			.HAlign(HAlign_Fill)
 			.Padding(2.0f, 24.0f, 2.0f, 2.0f)
 			[
 				SNew(STextBlock)
+				.AutoWrapText(true)
+				.Justification(ETextJustify::Center)
 				.Text(LOCTEXT("ConsoleVariablesEditorList_NoList", "No List to show. Try clearing the active search or adding some console variables to the list."))
 			]
 		]
-	]; 
+	];
+
+	EvaluateIfRowsPassFilters();
 }
 
 SConsoleVariablesEditorList::~SConsoleVariablesEditorList()
@@ -128,30 +136,63 @@ SConsoleVariablesEditorList::~SConsoleVariablesEditorList()
 	TreeViewPtr.Reset();
 }
 
+void SConsoleVariablesEditorList::SetupFilters()
+{
+	TArray<FString> SourceFilterTypes =
+	{
+		"Constructor",
+		"Scalability",
+		"Game Setting",
+		"Project Setting",
+		"System Settings ini",
+		"Device Profile",
+		"Game Override",
+		"Console Variables ini",
+		"Command line",
+		"Code",
+		"Console"
+	};
+
+	for (const FString& Type : SourceFilterTypes)
+	{
+		ShowFilters.Add(MakeShared<FConsoleVariablesEditorListFilter_SourceText>(Type));
+	}
+}
+
 TSharedRef<SWidget> SConsoleVariablesEditorList::BuildShowOptionsMenu()
 {
 	FMenuBuilder ShowOptionsMenuBuilder = FMenuBuilder(true, nullptr);
 
-	ShowOptionsMenuBuilder.BeginSection("AssetThumbnails", LOCTEXT("ShowHeading", "Show"));
+	ShowOptionsMenuBuilder.BeginSection("AssetThumbnails", LOCTEXT("ShowOptionsShowSectionHeading", "Show"));
 	{
 		// Add mode filters
-		auto AddFiltersLambda = [&ShowOptionsMenuBuilder](const FText& FilterTitle, const FText& FilterTooltip)
+		auto AddFiltersLambda = [this, &ShowOptionsMenuBuilder](const TSharedRef<IConsoleVariablesEditorListFilter>& InFilter)
 		{
+			const FString& FilterName = InFilter->GetFilterName();
+			
 			ShowOptionsMenuBuilder.AddMenuEntry(
-			   FilterTitle,
-			   FilterTooltip,
+			   InFilter->GetFilterButtonLabel(),
+			   InFilter->GetFilterButtonToolTip(),
 			   FSlateIcon(),
 			   FUIAction(
-				   //FExecuteAction::CreateRaw( this, &SConsoleVariablesEditorList::ToggleFilterActive ),
-				   //FCanExecuteAction(),
-				   //FIsActionChecked::CreateRaw( this, &SConsoleVariablesEditorList::IsFilterActive )
+				   FExecuteAction::CreateLambda(
+				   	[this, FilterName]()
+					   {
+						   ToggleFilterActive(FilterName);
+					   }
+					),
+				   FCanExecuteAction(),
+				   FIsActionChecked::CreateSP( InFilter, &IConsoleVariablesEditorListFilter::GetIsFilterActive )
 			   ),
 			   NAME_None,
 			   EUserInterfaceActionType::ToggleButton
 		   );
 		};
 
-		//for (Filter : Filters) { AddFiltersLambda(Filter.Title, Filter.TooltipText); }
+		for (const TSharedRef<IConsoleVariablesEditorListFilter>& Filter : ShowFilters)
+		{
+			AddFiltersLambda(Filter);
+		}
 	}
 	ShowOptionsMenuBuilder.EndSection();
 	
@@ -159,6 +200,7 @@ TSharedRef<SWidget> SConsoleVariablesEditorList::BuildShowOptionsMenu()
 	{
 		// Add commands
 
+		// Save this for later when folders are added
 		/*ShowOptionsMenuBuilder.AddMenuEntry(
 			LOCTEXT("CollapseAll", "Collapse All"),
 			LOCTEXT("ConsoleVariablesEditorList_CollapseAll_Tooltip", "Collapse all expanded actor groups in the Modified Actors list."),
@@ -194,18 +236,11 @@ void SConsoleVariablesEditorList::FlushMemory(const bool bShouldKeepMemoryAlloca
 	}
 }
 
-void SConsoleVariablesEditorList::RefreshScroll() const
-{
-	TreeViewPtr->RequestListRefresh();
-}
-
-void SConsoleVariablesEditorList::RefreshList(const FString& InConsoleCommandToScrollTo)
+void SConsoleVariablesEditorList::RebuildList(const FString& InConsoleCommandToScrollTo)
 {
 	GenerateTreeView();
 
-	// Enforce Sort
-	const FName& SortingName = GetActiveSortingColumnName();
-	ExecuteSort(SortingName, GetSortModeForColumn(SortingName));
+	RefreshList();
 
 	if (!InConsoleCommandToScrollTo.IsEmpty())
 	{
@@ -228,6 +263,24 @@ void SConsoleVariablesEditorList::RefreshList(const FString& InConsoleCommandToS
 	}
 }
 
+void SConsoleVariablesEditorList::RefreshList()
+{
+	if (TreeViewRootObjects.Num() > 0)
+	{
+		// Apply last search
+		ExecuteListViewSearchOnAllRows(GetSearchStringFromSearchInputField());
+
+		// Enforce Sort
+		const FName& SortingName = GetActiveSortingColumnName();
+		ExecuteSort(SortingName, GetSortModeForColumn(SortingName));
+
+		// Show/Hide rows based on SetBy changes and filter settings
+		EvaluateIfRowsPassFilters();
+
+		TreeViewPtr->RequestTreeRefresh();
+	}
+}
+
 TArray<FConsoleVariablesEditorListRowPtr> SConsoleVariablesEditorList::GetSelectedTreeViewItems() const
 {
 	return TreeViewPtr->GetSelectedItems();
@@ -307,13 +360,7 @@ void SConsoleVariablesEditorList::GenerateTreeView()
 		}
 	}
 
-	if (TreeViewRootObjects.Num() > 0)
-	{
-		TreeViewPtr->RequestListRefresh();
-
-		// Apply last search
-		ExecuteListViewSearchOnAllRows(GetSearchStringFromSearchInputField());
-	}
+	TreeViewPtr->RequestTreeRefresh();
 }
 
 TSharedPtr<SHeaderRow> SConsoleVariablesEditorList::GenerateHeaderRow()
@@ -519,6 +566,40 @@ void SConsoleVariablesEditorList::OnListItemCheckBoxStateChange(const ECheckBoxS
 	}
 }
 
+void SConsoleVariablesEditorList::ToggleFilterActive(const FString& FilterName)
+{
+	if (const TSharedRef<IConsoleVariablesEditorListFilter>* Match =
+		Algo::FindByPredicate(ShowFilters,
+		[&FilterName](TSharedRef<IConsoleVariablesEditorListFilter> Comparator)
+		{
+			return Comparator->GetFilterName().Equals(FilterName);
+		}))
+	{
+		const TSharedRef<IConsoleVariablesEditorListFilter> Filter = *Match;
+		Filter->ToggleFilterActive();
+
+		EvaluateIfRowsPassFilters();
+	}
+}
+
+void SConsoleVariablesEditorList::EvaluateIfRowsPassFilters()
+{
+	for (const FConsoleVariablesEditorListRowPtr& Row : TreeViewRootObjects)
+	{
+		if (Row.IsValid() && Row->GetRowType() == FConsoleVariablesEditorListRow::SingleCommand)
+		{
+			Row->SetDoesRowPassFilters(Algo::FindByPredicate(
+				ShowFilters,
+				[&Row](const TSharedRef<IConsoleVariablesEditorListFilter>& Filter)
+				{
+					return Filter->GetIsFilterActive() && Filter->DoesItemPassFilter(Row);
+				}) != nullptr);
+		}
+	}
+
+	TreeViewPtr->RequestTreeRefresh();
+}
+
 EColumnSortMode::Type SConsoleVariablesEditorList::GetSortModeForColumn(FName InColumnName) const
 {
 	EColumnSortMode::Type ColumnSortMode = EColumnSortMode::None;
@@ -627,4 +708,3 @@ void SConsoleVariablesEditorList::SetChildExpansionRecursively(const FConsoleVar
 };
 
 #undef LOCTEXT_NAMESPACE
-
diff --git a/Engine/Plugins/Experimental/ConsoleVariables/Source/ConsoleVariablesEditor/Private/Views/List/SConsoleVariablesEditorList.h b/Engine/Plugins/Experimental/ConsoleVariables/Source/ConsoleVariablesEditor/Private/Views/List/SConsoleVariablesEditorList.h
--- a/Engine/Plugins/Experimental/ConsoleVariables/Source/ConsoleVariablesEditor/Private/Views/List/SConsoleVariablesEditorList.h
+++ b/Engine/Plugins/Experimental/ConsoleVariables/Source/ConsoleVariablesEditor/Private/Views/List/SConsoleVariablesEditorList.h
@@ -3,6 +3,7 @@
 #pragma once
 
 #include "ConsoleVariablesEditorListRow.h"
+#include "ConsoleVariablesEditorListFilters/IConsoleVariablesEditorListFilter.h"
 
 #include "Widgets/SCompoundWidget.h"
 #include "Widgets/Views/SHeaderRow.h"
@@ -28,13 +29,14 @@ public:
 
 	virtual ~SConsoleVariablesEditorList() override;
 
-	TSharedRef<SWidget> BuildShowOptionsMenu();
-	
-	void FlushMemory(const bool bShouldKeepMemoryAllocated);
+	/** Regenerate the list items and refresh the list. Call when adding or removing variables. */
+	void RebuildList(const FString& InConsoleCommandToScrollTo = "");
 
-	void RefreshScroll() const;
-
-	void RefreshList(const FString& InConsoleCommandToScrollTo = "");
+	/**
+	 * Refresh filters and sorting.
+	 * Useful for when the list state has gone stale but the variable count has not changed.
+	 */
+	void RefreshList();
 
 	[[nodiscard]] TArray<FConsoleVariablesEditorListRowPtr> GetSelectedTreeViewItems() const;
 
@@ -64,6 +66,9 @@ public:
 	
 	void OnListItemCheckBoxStateChange(const ECheckBoxState InNewState);
 
+	void ToggleFilterActive(const FString& FilterName);
+	void EvaluateIfRowsPassFilters();
+
 	// Sorting
 
 	const FName& GetActiveSortingColumnName() const
@@ -94,6 +99,12 @@ private:
 	TSharedPtr<SHeaderRow> HeaderRow;
 	TSharedPtr<SHeaderRow> GenerateHeaderRow();
 	ECheckBoxState HeaderCheckBoxState = ECheckBoxState::Checked;
+
+	void SetupFilters();
+
+	TSharedRef<SWidget> BuildShowOptionsMenu();
+	
+	void FlushMemory(const bool bShouldKeepMemoryAllocated);
 	
 	void SetAllGroupsCollapsed();
 
@@ -114,6 +125,8 @@ private:
 
 	void SetChildExpansionRecursively(const FConsoleVariablesEditorListRowPtr& InRow, const bool bNewIsExpanded) const;
 
+	TArray<TSharedRef<IConsoleVariablesEditorListFilter>> ShowFilters;
+
 	TSharedPtr<STreeView<FConsoleVariablesEditorListRowPtr>> TreeViewPtr;
 	
 	TArray<FConsoleVariablesEditorListRowPtr> TreeViewRootObjects;
diff --git a/Engine/Plugins/Experimental/ConsoleVariables/Source/ConsoleVariablesEditor/Private/Views/MainPanel/ConsoleVariablesEditorMainPanel.cpp b/Engine/Plugins/Experimental/ConsoleVariables/Source/ConsoleVariablesEditor/Private/Views/MainPanel/ConsoleVariablesEditorMainPanel.cpp
--- a/Engine/Plugins/Experimental/ConsoleVariables/Source/ConsoleVariablesEditor/Private/Views/MainPanel/ConsoleVariablesEditorMainPanel.cpp
+++ b/Engine/Plugins/Experimental/ConsoleVariables/Source/ConsoleVariablesEditor/Private/Views/MainPanel/ConsoleVariablesEditorMainPanel.cpp
@@ -61,15 +61,23 @@ void FConsoleVariablesEditorMainPanel::AddConsoleVariable(
 
 		Asset->AddOrSetConsoleVariableSavedValue(InConsoleCommand, InValue);
 
-		RefreshList(bScrollToNewRow ? InConsoleCommand : "");
+		RebuildList(bScrollToNewRow ? InConsoleCommand : "");
 	}
 }
 
-void FConsoleVariablesEditorMainPanel::RefreshList(const FString& InConsoleCommandToScrollTo) const
+void FConsoleVariablesEditorMainPanel::RebuildList(const FString& InConsoleCommandToScrollTo) const
 {
 	if (EditorList.IsValid())
 	{
-		EditorList->RefreshList(InConsoleCommandToScrollTo);
+		EditorList->RebuildList(InConsoleCommandToScrollTo);
+	}
+}
+
+void FConsoleVariablesEditorMainPanel::RefreshList() const
+{
+	if (EditorList.IsValid())
+	{
+		EditorList->RefreshList();
 	}
 }
 
@@ -141,7 +149,7 @@ void FConsoleVariablesEditorMainPanel::ImportPreset(const FAssetData& InPresetAs
 
 	if (EditingAsset && ImportPreset_Impl(InPresetAsset, EditingAsset))
 	{
-		EditorList->RefreshList();
+		EditorList->RebuildList();
 	}
 }
 
diff --git a/Engine/Plugins/Experimental/ConsoleVariables/Source/ConsoleVariablesEditor/Private/Views/MainPanel/ConsoleVariablesEditorMainPanel.h b/Engine/Plugins/Experimental/ConsoleVariables/Source/ConsoleVariablesEditor/Private/Views/MainPanel/ConsoleVariablesEditorMainPanel.h
--- a/Engine/Plugins/Experimental/ConsoleVariables/Source/ConsoleVariablesEditor/Private/Views/MainPanel/ConsoleVariablesEditorMainPanel.h
+++ b/Engine/Plugins/Experimental/ConsoleVariables/Source/ConsoleVariablesEditor/Private/Views/MainPanel/ConsoleVariablesEditorMainPanel.h
@@ -23,7 +23,15 @@ public:
 
 	void AddConsoleVariable(const FString& InConsoleCommand, const FString& InValue, const bool bScrollToNewRow = false) const;
 
-	void RefreshList(const FString& InConsoleCommandToScrollTo = "") const;
+	/** Regenerate the list items and refresh the list. Call when adding or removing variables. */
+	void RebuildList(const FString& InConsoleCommandToScrollTo = "") const;
+
+	/**
+	 * Refresh filters and sorting.
+	 * Useful for when the list state has gone stale but the variable count has not changed.
+	 */
+	void RefreshList() const;
+	
 	void UpdatePresetValuesForSave(TObjectPtr<UConsoleVariablesAsset> InAsset) const;
 
 	void RefreshMultiUserDetails() const;
