Console Variables Editor:

-Fix for: Can't remove Editor.ReflectEditorLevelVisibilityWithGame from cvar list

diff --git a/Engine/Plugins/Experimental/ConsoleVariables/Source/ConsoleVariablesEditor/Private/ConsoleVariablesEditorCommandInfo.h b/Engine/Plugins/Experimental/ConsoleVariables/Source/ConsoleVariablesEditor/Private/ConsoleVariablesEditorCommandInfo.h
--- a/Engine/Plugins/Experimental/ConsoleVariables/Source/ConsoleVariablesEditor/Private/ConsoleVariablesEditorCommandInfo.h
+++ b/Engine/Plugins/Experimental/ConsoleVariables/Source/ConsoleVariablesEditor/Private/ConsoleVariablesEditorCommandInfo.h
@@ -2,6 +2,7 @@
 
 #pragma once
 
+#include "Algo/Find.h"
 #include "CoreMinimal.h"
 #include "ConsoleVariablesEditorLog.h"
 #include "Editor.h"
@@ -16,18 +17,26 @@ struct FConsoleVariablesEditorCommandInfo
 {
 	GENERATED_BODY()
 
+	struct FStaticConsoleVariableFlagInfo
+	{
+		EConsoleVariableFlags Flag;
+		FText DisplayText;
+	};
+
 	FConsoleVariablesEditorCommandInfo()
 	{
 		ConsoleVariablePtr = nullptr;
 	}
 
 	FConsoleVariablesEditorCommandInfo(
-		const FString& InCommand, IConsoleVariable* InVariablePtr, const FString& InStartupValue, const FDelegateHandle& InOnVariableChangedCallbackHandle)
+		const FString& InCommand, IConsoleVariable* InVariablePtr, const FDelegateHandle& InOnVariableChangedCallbackHandle)
 	: Command(InCommand)
 	, ConsoleVariablePtr(InVariablePtr)
-	, StartupValueAsString(InStartupValue)
 	, OnVariableChangedCallbackHandle(InOnVariableChangedCallbackHandle)
-	{}
+	{
+		StartupValueAsString = ConsoleVariablePtr->GetString();
+		StartupSource = GetSource();
+	}
 
 	FORCEINLINE bool operator==(const FConsoleVariablesEditorCommandInfo& Comparator) const
 	{
@@ -65,8 +74,17 @@ struct FConsoleVariablesEditorCommandInfo
 		return ECVF_Default;
 	}
 
+	void ClearSourceFlags() const
+	{
+		for (const FStaticConsoleVariableFlagInfo StaticConsoleVariableFlagInfo : SupportedFlags)
+		{
+			ConsoleVariablePtr->ClearFlags(StaticConsoleVariableFlagInfo.Flag);
+		}
+	}
+
 	void SetSourceFlag(const EConsoleVariableFlags InSource) const
 	{
+		ClearSourceFlags();
 		ConsoleVariablePtr->SetFlags((EConsoleVariableFlags)InSource);
 	}
 
@@ -77,47 +95,18 @@ struct FConsoleVariablesEditorCommandInfo
 
 	static FText ConvertConsoleVariableSetByFlagToText(const EConsoleVariableFlags InFlag)
 	{
-		
-		switch (InFlag)
+		FText ReturnValue = LOCTEXT("UnknownSource", "<UNKNOWN>"); 
+
+		if (const FStaticConsoleVariableFlagInfo* Match = Algo::FindByPredicate(SupportedFlags,
+				[InFlag](const FStaticConsoleVariableFlagInfo& Comparator)
+				{
+					return Comparator.Flag == InFlag;
+				}))
 		{
-			case (EConsoleVariableFlags::ECVF_SetByConstructor):
-				return LOCTEXT("SetByConstructor", "Constructor");
-				
-			case (EConsoleVariableFlags::ECVF_SetByScalability):
-				return LOCTEXT("SetByScalability", "Scalability");
-				
-			case (EConsoleVariableFlags::ECVF_SetByGameSetting):
-				return LOCTEXT("SetByGameSetting", "Game Setting");
-				
-			case (EConsoleVariableFlags::ECVF_SetByProjectSetting):
-				return LOCTEXT("SetByProjectSetting", "Project Setting");
-				
-			case (EConsoleVariableFlags::ECVF_SetBySystemSettingsIni):
-				return LOCTEXT("SetBySystemSettingsIni", "System Settings ini");
-				
-			case (EConsoleVariableFlags::ECVF_SetByDeviceProfile):
-				return LOCTEXT("SetByDeviceProfile", "Device Profile");
-				
-			case (EConsoleVariableFlags::ECVF_SetByGameOverride):
-				return LOCTEXT("SetByGameOverride", "Game Override");
-				
-			case (EConsoleVariableFlags::ECVF_SetByConsoleVariablesIni):
-				return LOCTEXT("SetByConsoleVariablesIni", "Console Variables ini");
-				
-			case (EConsoleVariableFlags::ECVF_SetByCommandline):
-				return LOCTEXT("SetByCommandline", "Command line");
-				
-			case (EConsoleVariableFlags::ECVF_SetByCode):
-				return LOCTEXT("SetByCode", "Code");
-				
-			case (EConsoleVariableFlags::ECVF_SetByConsole):
-				return LOCTEXT("SetByConsole", "Console");
-
-			default:
-				break;;
+			ReturnValue = (*Match).DisplayText;
 		}
 		
-		return LOCTEXT("UnknownSource", "<UNKNOWN>");
+		return ReturnValue;
 	}
 
 	bool IsCurrentValueDifferentFromInputValue(const FString& InValueToCompare) const
@@ -142,6 +131,21 @@ struct FConsoleVariablesEditorCommandInfo
 	EConsoleVariableFlags StartupSource;
 	
 	FDelegateHandle OnVariableChangedCallbackHandle;
+
+	static const inline TArray<FStaticConsoleVariableFlagInfo> SupportedFlags =
+	{
+		{ EConsoleVariableFlags::ECVF_SetByConstructor, LOCTEXT("SetByConstructor", "Constructor") },
+		{ EConsoleVariableFlags::ECVF_SetByScalability, LOCTEXT("SetByScalability", "Scalability") },
+		{ EConsoleVariableFlags::ECVF_SetByGameSetting, LOCTEXT("SetByGameSetting", "Game Setting") },
+		{ EConsoleVariableFlags::ECVF_SetByProjectSetting, LOCTEXT("SetByProjectSetting", "Project Setting") },
+		{ EConsoleVariableFlags::ECVF_SetBySystemSettingsIni, LOCTEXT("SetBySystemSettingsIni", "System Settings ini") },
+		{ EConsoleVariableFlags::ECVF_SetByDeviceProfile, LOCTEXT("SetByDeviceProfile", "Device Profile") },
+		{ EConsoleVariableFlags::ECVF_SetByGameOverride, LOCTEXT("SetByGameOverride", "Game Override") },
+		{ EConsoleVariableFlags::ECVF_SetByConsoleVariablesIni, LOCTEXT("SetByConsoleVariablesIni", "Console Variables ini") },
+		{ EConsoleVariableFlags::ECVF_SetByCommandline, LOCTEXT("SetByCommandline", "Command line") },
+		{ EConsoleVariableFlags::ECVF_SetByCode, LOCTEXT("SetByCode", "Code") },
+		{ EConsoleVariableFlags::ECVF_SetByConsole, LOCTEXT("SetByConsole", "Console") }
+	};
 };
 
 #undef LOCTEXT_NAMESPACE
diff --git a/Engine/Plugins/Experimental/ConsoleVariables/Source/ConsoleVariablesEditor/Private/Views/List/SConsoleVariablesEditorListRow.cpp b/Engine/Plugins/Experimental/ConsoleVariables/Source/ConsoleVariablesEditor/Private/Views/List/SConsoleVariablesEditorListRow.cpp
--- a/Engine/Plugins/Experimental/ConsoleVariables/Source/ConsoleVariablesEditor/Private/Views/List/SConsoleVariablesEditorListRow.cpp
+++ b/Engine/Plugins/Experimental/ConsoleVariables/Source/ConsoleVariablesEditor/Private/Views/List/SConsoleVariablesEditorListRow.cpp
@@ -439,13 +439,19 @@ TSharedRef<SWidget> SConsoleVariablesEditorListRow::GenerateValueCellWidget(cons
 			SNew(SHorizontalBox)
 			.ToolTipText_Lambda([this, PinnedItem] ()
 			{
-				return FText::Format(
-							ValueWidgetToolTipFormatText,
-							FText::FromString(ValueChildInputWidget->GetCachedValue()),
-							FText::FromString(PinnedItem->GetPresetValue()),
-							FText::FromString(PinnedItem->GetCommandInfo().Pin()->StartupValueAsString),
-							FConsoleVariablesEditorCommandInfo::ConvertConsoleVariableSetByFlagToText(PinnedItem->GetCommandInfo().Pin()->StartupSource)
-				);
+				if (const TSharedPtr<FConsoleVariablesEditorCommandInfo> PinnedCommand = PinnedItem->GetCommandInfo().Pin();
+					PinnedCommand.IsValid())
+				{
+					return FText::Format(
+								ValueWidgetToolTipFormatText,
+								FText::FromString(ValueChildInputWidget->GetCachedValue()),
+								FText::FromString(PinnedItem->GetPresetValue()),
+								FText::FromString(PinnedCommand->StartupValueAsString),
+								FConsoleVariablesEditorCommandInfo::ConvertConsoleVariableSetByFlagToText(PinnedCommand->StartupSource)
+					);
+				}
+
+				return FText::GetEmpty();
 			});
 
 		FinalValueWidget->AddSlot()
