Compile error fixes

#[fyi] Jason.Walter

diff --git a/Engine/Plugins/Experimental/ConsoleVariables/Source/ConsoleVariablesEditor/Private/ConsoleVariablesEditorCommandInfo.h b/Engine/Plugins/Experimental/ConsoleVariables/Source/ConsoleVariablesEditor/Private/ConsoleVariablesEditorCommandInfo.h
--- a/Engine/Plugins/Experimental/ConsoleVariables/Source/ConsoleVariablesEditor/Private/ConsoleVariablesEditorCommandInfo.h
+++ b/Engine/Plugins/Experimental/ConsoleVariables/Source/ConsoleVariablesEditor/Private/ConsoleVariablesEditorCommandInfo.h
@@ -76,7 +76,7 @@ struct FConsoleVariablesEditorCommandInfo
 
 	void ClearSourceFlags() const
 	{
-		for (const FStaticConsoleVariableFlagInfo StaticConsoleVariableFlagInfo : SupportedFlags)
+		for (const FStaticConsoleVariableFlagInfo& StaticConsoleVariableFlagInfo : SupportedFlags)
 		{
 			ConsoleVariablePtr->ClearFlags(StaticConsoleVariableFlagInfo.Flag);
 		}
diff --git a/Engine/Plugins/Experimental/ConsoleVariables/Source/ConsoleVariablesEditor/Private/Views/List/ConsoleVariablesEditorListFilters/ConsoleVariablesEditorListFilter_SourceText.h b/Engine/Plugins/Experimental/ConsoleVariables/Source/ConsoleVariablesEditor/Private/Views/List/ConsoleVariablesEditorListFilters/ConsoleVariablesEditorListFilter_SourceText.h
--- a/Engine/Plugins/Experimental/ConsoleVariables/Source/ConsoleVariablesEditor/Private/Views/List/ConsoleVariablesEditorListFilters/ConsoleVariablesEditorListFilter_SourceText.h
+++ b/Engine/Plugins/Experimental/ConsoleVariables/Source/ConsoleVariablesEditor/Private/Views/List/ConsoleVariablesEditorListFilters/ConsoleVariablesEditorListFilter_SourceText.h
@@ -32,7 +32,7 @@ public:
 
 	virtual bool DoesItemPassFilter(const FConsoleVariablesEditorListRowPtr& InItem) override
 	{
-		if (InItem.IsValid(); const TSharedPtr<FConsoleVariablesEditorCommandInfo> PinnedCommand = InItem->GetCommandInfo().Pin();)
+		if (InItem.IsValid(); const TSharedPtr<FConsoleVariablesEditorCommandInfo> PinnedCommand = InItem->GetCommandInfo().Pin())
 		{
 			const FString& StringToSearch = PinnedCommand->GetSourceAsText().ToString();
 
