Console Variables Editor:

-Fix for: nDisplay crashes when attempting to launch a config that is not in the selected level

#[fyi] Jason.Walter

diff --git a/Engine/Plugins/Experimental/ConsoleVariables/Source/ConsoleVariablesEditor/Private/ConsoleVariablesAsset.cpp b/Engine/Plugins/Experimental/ConsoleVariables/Source/ConsoleVariablesEditor/Private/ConsoleVariablesAsset.cpp
--- a/Engine/Plugins/Experimental/ConsoleVariables/Source/ConsoleVariablesEditor/Private/ConsoleVariablesAsset.cpp
+++ b/Engine/Plugins/Experimental/ConsoleVariables/Source/ConsoleVariablesEditor/Private/ConsoleVariablesAsset.cpp
@@ -18,12 +18,9 @@ void UConsoleVariablesAsset::ReplaceSavedCommandsAndValues(const TMap<FString, F
 
 bool UConsoleVariablesAsset::FindSavedValueByCommandString(const FString& InCommandString, FString& OutValue) const
 {
-	TArray<FString> KeyArray;
-	SavedCommandsAndValues.GetKeys(KeyArray);
-
-	if (const bool bMatchFound = KeyArray.Contains(InCommandString))
+	if (const FString* Match = SavedCommandsAndValues.Find(InCommandString))
 	{
-		OutValue = SavedCommandsAndValues[InCommandString];
+		OutValue = *Match;
 		return true;
 	}
 	
@@ -37,10 +34,7 @@ void UConsoleVariablesAsset::AddOrSetConsoleVariableSavedValue(const FString& In
 
 bool UConsoleVariablesAsset::RemoveConsoleVariable(const FString& InCommandString)
 {
-	TArray<FString> KeyArray;
-	SavedCommandsAndValues.GetKeys(KeyArray);
-
-	if (const bool bMatchFound = KeyArray.Contains(InCommandString))
+	if (const FString* Match = SavedCommandsAndValues.Find(InCommandString))
 	{
 		return SavedCommandsAndValues.Remove(InCommandString) > 0;
 	}
