Console Variables Editor:

-Fix for: changed values are set continuously on the node resulting in nDisplay crash

#[fyi] Jason.Walter

diff --git a/Engine/Plugins/Experimental/ConsoleVariables/Source/ConsoleVariablesEditor/Private/ConsoleVariablesEditorModule.cpp b/Engine/Plugins/Experimental/ConsoleVariables/Source/ConsoleVariablesEditor/Private/ConsoleVariablesEditorModule.cpp
--- a/Engine/Plugins/Experimental/ConsoleVariables/Source/ConsoleVariablesEditor/Private/ConsoleVariablesEditorModule.cpp
+++ b/Engine/Plugins/Experimental/ConsoleVariables/Source/ConsoleVariablesEditor/Private/ConsoleVariablesEditorModule.cpp
@@ -217,6 +217,24 @@ void FConsoleVariablesEditorModule::OnConsoleVariableChanged(IConsoleVariable* C
 	}
 }
 
+void FConsoleVariablesEditorModule::OnRemoteCvarChanged(const FString InName, const FString InValue)
+{
+	UE_LOG(LogConsoleVariablesEditor, Display, TEXT("Remote set console variable %s = %s"), *InName, *InValue);
+
+	if (GetMutableDefault<UConcertCVarSynchronization>()->bSyncCVarTransactions)
+	{
+		if (const TWeakPtr<FConsoleVariablesEditorCommandInfo> CommandInfo =
+			FindCommandInfoByName(InName); CommandInfo.IsValid())
+		{
+			if (CommandInfo.Pin()->IsCurrentValueDifferentFromInputValue(InValue))
+			{
+				GEngine->Exec(FConsoleVariablesEditorCommandInfo::GetCurrentWorld(),
+					*FString::Printf(TEXT("%s %s"), *InName, *InValue));
+			}
+		}
+	}
+}
+
 TObjectPtr<UConsoleVariablesAsset> FConsoleVariablesEditorModule::AllocateTransientPreset()
 {
 	static const TCHAR* PackageName = TEXT("/Temp/ConsoleVariablesEditor/PendingConsoleVariablesCollections");
diff --git a/Engine/Plugins/Experimental/ConsoleVariables/Source/ConsoleVariablesEditor/Private/ConsoleVariablesEditorModule.h b/Engine/Plugins/Experimental/ConsoleVariables/Source/ConsoleVariablesEditor/Private/ConsoleVariablesEditorModule.h
--- a/Engine/Plugins/Experimental/ConsoleVariables/Source/ConsoleVariablesEditor/Private/ConsoleVariablesEditorModule.h
+++ b/Engine/Plugins/Experimental/ConsoleVariables/Source/ConsoleVariablesEditor/Private/ConsoleVariablesEditorModule.h
@@ -43,6 +43,7 @@ public:
 	void SetEditingAsset(const TObjectPtr<UConsoleVariablesAsset> InEditingAsset);
 
 	void SendMultiUserConsoleVariableChange(const FString& InVariableName, const FString& InValueAsString);
+	void OnRemoteCvarChanged(const FString InName, const FString InValue);
 	
 private:
 
diff --git a/Engine/Plugins/Experimental/ConsoleVariables/Source/ConsoleVariablesEditor/Private/Views/MainPanel/ConsoleVariablesEditorMainPanel.cpp b/Engine/Plugins/Experimental/ConsoleVariables/Source/ConsoleVariablesEditor/Private/Views/MainPanel/ConsoleVariablesEditorMainPanel.cpp
--- a/Engine/Plugins/Experimental/ConsoleVariables/Source/ConsoleVariablesEditor/Private/Views/MainPanel/ConsoleVariablesEditorMainPanel.cpp
+++ b/Engine/Plugins/Experimental/ConsoleVariables/Source/ConsoleVariablesEditor/Private/Views/MainPanel/ConsoleVariablesEditorMainPanel.cpp
@@ -4,8 +4,9 @@
 
 #include "ConcertMessages.h"
 #include "ConsoleVariablesAsset.h"
-#include "ConsoleVariablesEditorLog.h"
 #include "ConsoleVariablesEditorCommandInfo.h"
+#include "ConsoleVariablesEditorLog.h"
+#include "ConsoleVariablesEditorModule.h"
 #include "Views/List/ConsoleVariablesEditorList.h"
 #include "Views/MainPanel/SConsoleVariablesEditorMainPanel.h"
 
@@ -188,11 +189,7 @@ void FConsoleVariablesEditorMainPanel::OnConnectionChanged(EConcertConnectionSta
 
 void FConsoleVariablesEditorMainPanel::OnRemoteCvarChange(const FString InName, const FString InValue)
 {
-	UE_LOG(LogConsoleVariablesEditor, Display, TEXT("Remote set console variable %s = %s"), *InName, *InValue);
+	FConsoleVariablesEditorModule& ConsoleVariablesEditorModule = FConsoleVariablesEditorModule::Get();
 
-	if (GetMutableDefault<UConcertCVarSynchronization>()->bSyncCVarTransactions)
-	{
-		GEngine->Exec(FConsoleVariablesEditorCommandInfo::GetCurrentWorld(),
-					  *FString::Printf(TEXT("%s %s"), *InName, *InValue));
-	}
+	ConsoleVariablesEditorModule.OnRemoteCvarChanged(InName, InValue);
 }
