Console Variables Editor:

-Make custom widget for row tooltips for hovering over rows in the list and suggestion box

#[fyi] Jason.Walter

diff --git a/Engine/Plugins/Experimental/ConsoleVariables/Source/ConsoleVariablesEditor/Private/Views/List/SConsoleVariablesEditorListRow.cpp b/Engine/Plugins/Experimental/ConsoleVariables/Source/ConsoleVariablesEditor/Private/Views/List/SConsoleVariablesEditorListRow.cpp
--- a/Engine/Plugins/Experimental/ConsoleVariables/Source/ConsoleVariablesEditor/Private/Views/List/SConsoleVariablesEditorListRow.cpp
+++ b/Engine/Plugins/Experimental/ConsoleVariables/Source/ConsoleVariablesEditor/Private/Views/List/SConsoleVariablesEditorListRow.cpp
@@ -5,9 +5,9 @@
 #include "ConsoleVariablesEditorModule.h"
 #include "ConsoleVariablesEditorStyle.h"
 #include "SConsoleVariablesEditorListValueInput.h"
+#include "Views/Widgets/SConsoleVariablesEditorTooltipWidget.h"
 
 #include "Input/DragAndDrop.h"
-#include "Kismet/KismetMathLibrary.h"
 #include "Styling/AppStyle.h"
 #include "Styling/StyleColors.h"
 #include "Widgets/Input/SCheckBox.h"
@@ -64,6 +64,14 @@ void SConsoleVariablesEditorListRow::Construct(
 
 	Item = InRow;
 	const FConsoleVariablesEditorListRowPtr PinnedItem = Item.Pin();
+	const TSharedPtr<FConsoleVariablesEditorCommandInfo> PinnedCommand = PinnedItem->GetCommandInfo().Pin();
+	
+	check (PinnedCommand.IsValid());
+
+	if (!PinnedCommand->ConsoleVariablePtr)
+	{
+		return;
+	}
 
 	// Set up flash animation
 	FlashAnimation = FCurveSequence(0.f, FlashAnimationDuration, ECurveEaseFunction::QuadInOut);
@@ -113,8 +121,6 @@ TSharedRef<SWidget> SConsoleVariablesEditorListRow::GenerateWidgetForColumn(cons
 				+SOverlay::Slot()
 				[
 					SNew(SBorder)
-					.ToolTipText(FText::FromString(PinnedItem->GetCommandInfo().Pin()->ConsoleVariablePtr ?
-						FString(PinnedItem->GetCommandInfo().Pin()->ConsoleVariablePtr->GetHelp()) : ""))
 					.HAlign(HAlign_Fill)
 					.VAlign(VAlign_Center)
 					.BorderImage(GetBorderImage(PinnedItem->GetRowType()))
@@ -152,6 +158,7 @@ void SConsoleVariablesEditorListRow::OnMouseLeave(const FPointerEvent& MouseEven
 SConsoleVariablesEditorListRow::~SConsoleVariablesEditorListRow()
 {
 	Item.Reset();
+	HoverToolTip.Reset();
 
 	FlashImages.Empty();
 
@@ -229,7 +236,7 @@ TOptional<EItemDropZone> SConsoleVariablesEditorListRow::HandleCanAcceptDrop(con
 			FText::FromString(TargetItem->GetCommandInfo().Pin()->Command)
 		);
 	
-	Operation->SetToolTip(
+		Operation->SetToolTip(
 		DropPermittedText,
 		FAppStyle::Get().GetBrush("Graph.ConnectorFeedback.OK")
 	);
@@ -352,9 +359,16 @@ TSharedRef<SWidget> SConsoleVariablesEditorListRow::GenerateCells(const FName& I
 	}
 	if (InColumnName.IsEqual(SConsoleVariablesEditorList::VariableNameColumnName))
 	{
+		if (!HoverToolTip.IsValid())
+		{
+			HoverToolTip = SConsoleVariablesEditorTooltipWidget::MakeTooltip(
+							PinnedItem->GetCommandInfo().Pin()->Command,
+							PinnedItem->GetCommandInfo().Pin()->ConsoleVariablePtr->GetHelp());
+		}
 		return  SNew(STextBlock)
-				.Visibility(EVisibility::SelfHitTestInvisible)
-				.Text(FText::FromString(PinnedItem->GetCommandInfo().Pin()->Command));
+				.Visibility(EVisibility::Visible)
+				.Text(FText::FromString(PinnedItem->GetCommandInfo().Pin()->Command))
+				.ToolTip(HoverToolTip);
 	}
 	if (InColumnName.IsEqual(SConsoleVariablesEditorList::ValueColumnName))
 	{
diff --git a/Engine/Plugins/Experimental/ConsoleVariables/Source/ConsoleVariablesEditor/Private/Views/List/SConsoleVariablesEditorListRow.h b/Engine/Plugins/Experimental/ConsoleVariables/Source/ConsoleVariablesEditor/Private/Views/List/SConsoleVariablesEditorListRow.h
--- a/Engine/Plugins/Experimental/ConsoleVariables/Source/ConsoleVariablesEditor/Private/Views/List/SConsoleVariablesEditorListRow.h
+++ b/Engine/Plugins/Experimental/ConsoleVariables/Source/ConsoleVariablesEditor/Private/Views/List/SConsoleVariablesEditorListRow.h
@@ -10,8 +10,8 @@
 #include "Widgets/Input/SButton.h"
 #include "Widgets/Input/SNumericEntryBox.h"
 #include "Widgets/Layout/SBorder.h"
-#include "Widgets/Layout/SSplitter.h"
 #include "Widgets/SCompoundWidget.h"
+#include "Widgets/IToolTip.h"
 
 class SConsoleVariablesEditorListValueInput;
 class SConsoleVariablesEditorListRowHoverWidgets;
@@ -70,6 +70,8 @@ public:
 private:
 	
 	TWeakPtr<FConsoleVariablesEditorListRow> Item;
+	
+	TSharedPtr<IToolTip> HoverToolTip;
 
 	TArray<TSharedPtr<SImage>> FlashImages;
 
diff --git a/Engine/Plugins/Experimental/ConsoleVariables/Source/ConsoleVariablesEditor/Private/Views/MainPanel/SConsoleVariablesEditorCustomConsoleInputBox.cpp b/Engine/Plugins/Experimental/ConsoleVariables/Source/ConsoleVariablesEditor/Private/Views/MainPanel/SConsoleVariablesEditorCustomConsoleInputBox.cpp
--- a/Engine/Plugins/Experimental/ConsoleVariables/Source/ConsoleVariablesEditor/Private/Views/MainPanel/SConsoleVariablesEditorCustomConsoleInputBox.cpp
+++ b/Engine/Plugins/Experimental/ConsoleVariables/Source/ConsoleVariablesEditor/Private/Views/MainPanel/SConsoleVariablesEditorCustomConsoleInputBox.cpp
@@ -11,6 +11,7 @@
 #include "OutputLog/Public/OutputLogModule.h"
 #include "Styling/AppStyle.h"
 #include "Styling/StyleColors.h"
+#include "Views/Widgets/SConsoleVariablesEditorTooltipWidget.h"
 #include "Widgets/Input/SEditableTextBox.h"
 #include "Widgets/Layout/SBorder.h"
 #include "Widgets/Layout/SBox.h"
@@ -71,14 +72,30 @@ void SConsoleVariablesEditorCustomConsoleInputBox::Construct(
 						SanitizedText.ReplaceInline(TEXT("\r"), TEXT(" "), ESearchCase::CaseSensitive);
 						SanitizedText.ReplaceInline(TEXT("\n"), TEXT(" "), ESearchCase::CaseSensitive);
 
+						const FText DisplayText = FText::FromString(SanitizedText);
+						
+						FConsoleVariablesEditorModule& ConsoleVariablesEditorModule = FConsoleVariablesEditorModule::Get();
+
+						TSharedPtr<IToolTip> ToolTip;
+
+						if (const TWeakPtr<FConsoleVariablesEditorCommandInfo> MatchingCommand =
+							ConsoleVariablesEditorModule.FindCommandInfoByName(SanitizedText);
+							MatchingCommand.IsValid() && MatchingCommand.Pin()->ConsoleVariablePtr)
+						{
+							ToolTip = SConsoleVariablesEditorTooltipWidget::MakeTooltip(
+								SanitizedText,
+								MatchingCommand.Pin()->ConsoleVariablePtr->GetHelp());
+						}
+
 						return
 							SNew(STableRow< TSharedPtr<FString> >, OwnerTable)
 							[
 								SNew(STextBlock)
-								.Text(FText::FromString(SanitizedText))
+								.Text(DisplayText)
 								.TextStyle(FEditorStyle::Get(), "Log.Normal")
 								.HighlightText(Suggestions.SuggestionsHighlight)
 								.ColorAndOpacity(FSlateColor::UseForeground())
+								.ToolTip(ToolTip.IsValid() ? ToolTip : nullptr)
 							];
 					})
 					.OnSelectionChanged_Lambda([this] (TSharedPtr<FString> NewValue, ESelectInfo::Type SelectInfo)
diff --git a/Engine/Plugins/Experimental/ConsoleVariables/Source/ConsoleVariablesEditor/Private/Views/Widgets/SConsoleVariablesEditorTooltipWidget.h b/Engine/Plugins/Experimental/ConsoleVariables/Source/ConsoleVariablesEditor/Private/Views/Widgets/SConsoleVariablesEditorTooltipWidget.h
new file mode 100644
--- /dev/null
+++ b/Engine/Plugins/Experimental/ConsoleVariables/Source/ConsoleVariablesEditor/Private/Views/Widgets/SConsoleVariablesEditorTooltipWidget.h
@@ -0,0 +1,79 @@
+// Copyright Epic Games, Inc. All Rights Reserved.
+
+#pragma once
+
+#include "Styling/AppStyle.h"
+#include "Widgets/Layout/SBorder.h"
+#include "Widgets/SBoxPanel.h"
+#include "Widgets/SCompoundWidget.h"
+#include "Widgets/SToolTip.h"
+#include "Widgets/Text/STextBlock.h"
+
+class SConsoleVariablesEditorTooltipWidget : public SCompoundWidget
+{
+public:
+	
+	SLATE_BEGIN_ARGS(SConsoleVariablesEditorTooltipWidget)
+	{}
+
+	SLATE_END_ARGS()
+	
+	void Construct(const FArguments& InArgs, const FString InCommandName, const FString InHelpString)
+	{
+		TSharedPtr<SVerticalBox> VerticalBox;
+		
+		ChildSlot
+		[
+			SAssignNew(VerticalBox, SVerticalBox)
+
+			+SVerticalBox::Slot()
+			.AutoHeight()
+			.Padding(0, 0, 0, 4)
+			[
+				SNew(SBorder)
+				.Padding(6)
+				.BorderImage(FAppStyle::Get().GetBrush("ContentBrowser.TileViewTooltip.ContentBorder"))
+				[
+					SNew(STextBlock)
+					.Text(FText::FromString(InCommandName))
+					.Font(FAppStyle::Get().GetFontStyle("ContentBrowser.TileViewTooltip.NameFont"))
+					.AutoWrapText(true)
+				]
+			]
+		];
+
+		if (!InHelpString.IsEmpty())
+		{
+			VerticalBox->AddSlot()
+			.AutoHeight()
+			[
+				SNew(SBorder)
+				.Padding(6)
+				.BorderImage(FAppStyle::Get().GetBrush("ContentBrowser.TileViewTooltip.ContentBorder"))
+				[
+					SNew(STextBlock)
+					.Text(FText::FromString(InHelpString))
+					.AutoWrapText(true)
+				]
+			];
+		}
+	}
+
+	static TSharedRef<SToolTip> MakeTooltip(const FString InCommandName, const FString InHelpString)
+	{
+		return
+			SNew(SToolTip)
+			.BorderImage(FAppStyle::Get().GetBrush("ContentBrowser.TileViewTooltip.NonContentBorder"))
+			[
+				SNew(SBox)
+				.WidthOverride(TooltipWindowWidth)
+				[
+					SNew(SConsoleVariablesEditorTooltipWidget, InCommandName, InHelpString)
+				]
+			];
+	}
+
+private:
+
+	static constexpr float TooltipWindowWidth = 500.f;
+};
