Console Variables Editor:

-Filter non-variable console commands from the custom console variables input box

#[fyi] Jason.Walter

diff --git a/Engine/Plugins/Experimental/ConsoleVariables/Source/ConsoleVariablesEditor/Private/Views/MainPanel/SConsoleVariablesEditorCustomConsoleInputBox.cpp b/Engine/Plugins/Experimental/ConsoleVariables/Source/ConsoleVariablesEditor/Private/Views/MainPanel/SConsoleVariablesEditorCustomConsoleInputBox.cpp
--- a/Engine/Plugins/Experimental/ConsoleVariables/Source/ConsoleVariablesEditor/Private/Views/MainPanel/SConsoleVariablesEditorCustomConsoleInputBox.cpp
+++ b/Engine/Plugins/Experimental/ConsoleVariables/Source/ConsoleVariablesEditor/Private/Views/MainPanel/SConsoleVariablesEditorCustomConsoleInputBox.cpp
@@ -221,6 +221,11 @@ void SConsoleVariablesEditorCustomConsoleInputBox::OnTextChanged(const FText& In
 				return;
 			}
 
+			if (!CVar->AsVariable())
+			{
+				return;
+			}
+
 			AutoCompleteList.Add(Name);
 		};
 
diff --git a/Engine/Plugins/Experimental/ConsoleVariables/Source/ConsoleVariablesEditor/Private/Views/MainPanel/SConsoleVariablesEditorMainPanel.cpp b/Engine/Plugins/Experimental/ConsoleVariables/Source/ConsoleVariablesEditor/Private/Views/MainPanel/SConsoleVariablesEditorMainPanel.cpp
--- a/Engine/Plugins/Experimental/ConsoleVariables/Source/ConsoleVariablesEditor/Private/Views/MainPanel/SConsoleVariablesEditorMainPanel.cpp
+++ b/Engine/Plugins/Experimental/ConsoleVariables/Source/ConsoleVariablesEditor/Private/Views/MainPanel/SConsoleVariablesEditorMainPanel.cpp
@@ -86,7 +86,8 @@ FReply SConsoleVariablesEditorMainPanel::ValidateConsoleInput(const FText& Commi
 	}
 	else
 	{
-		UE_LOG(LogConsoleVariablesEditor, Warning, TEXT("%hs: Input %s is not a recognized console command."),
+		UE_LOG(LogConsoleVariablesEditor, Warning,
+			TEXT("%hs: Input %s is not a recognized console variable. Note that only console variables are supported, not console commands."),
 			__FUNCTION__, *CommandString);
 	}
 
