Console Variables Editor:

-Change save type from TMap to structure and save checked state

#[fyi] Jason.Walter

diff --git a/Engine/Plugins/Experimental/ConsoleVariables/Source/ConsoleVariablesEditor/Private/ConsoleVariablesAsset.cpp b/Engine/Plugins/Experimental/ConsoleVariables/Source/ConsoleVariablesEditor/Private/ConsoleVariablesAsset.cpp
--- a/Engine/Plugins/Experimental/ConsoleVariables/Source/ConsoleVariablesEditor/Private/ConsoleVariablesAsset.cpp
+++ b/Engine/Plugins/Experimental/ConsoleVariables/Source/ConsoleVariablesEditor/Private/ConsoleVariablesAsset.cpp
@@ -11,14 +11,19 @@ void UConsoleVariablesAsset::SetVariableCollectionDescription(const FString& InV
 	VariableCollectionDescription = InVariableCollectionDescription;
 }
 
-void UConsoleVariablesAsset::ReplaceSavedCommandsAndValues(const TMap<FString, FString>& InMap)
+void UConsoleVariablesAsset::ReplaceSavedCommands(const TArray<FConsoleVariablesEditorAssetSaveData>& Replacement)
 {
-	SavedCommandsAndValues = InMap;
+	SavedCommands = Replacement;
 }
 
-bool UConsoleVariablesAsset::FindSavedValueByCommandString(const FString& InCommandString, FString& OutValue) const
+bool UConsoleVariablesAsset::FindSavedDataByCommandString(const FString& InCommandString, FConsoleVariablesEditorAssetSaveData& OutValue) const
 {
-	if (const FString* Match = SavedCommandsAndValues.Find(InCommandString))
+	if (const auto* Match = Algo::FindByPredicate(
+		SavedCommands,
+		[&InCommandString] (const FConsoleVariablesEditorAssetSaveData& Comparator)
+		{
+			return Comparator.CommandName.Equals(InCommandString);
+		}))
 	{
 		OutValue = *Match;
 		return true;
@@ -27,23 +32,25 @@ bool UConsoleVariablesAsset::FindSavedValueByCommandString(const FString& InComm
 	return false;
 }
 
-void UConsoleVariablesAsset::AddOrSetConsoleVariableSavedValue(const FString& InCommandString, const FString& InNewValue)
+void UConsoleVariablesAsset::AddOrSetConsoleVariableSavedData(const FConsoleVariablesEditorAssetSaveData& InData)
 {
-	SavedCommandsAndValues.Add(InCommandString, InNewValue);
+	RemoveConsoleVariable(InData.CommandName);
+	SavedCommands.Add(InData);
 }
 
 bool UConsoleVariablesAsset::RemoveConsoleVariable(const FString& InCommandString)
 {
-	if (const FString* Match = SavedCommandsAndValues.Find(InCommandString))
+	FConsoleVariablesEditorAssetSaveData ExistingData;
+	if (FindSavedDataByCommandString(InCommandString, ExistingData))
 	{
-		return SavedCommandsAndValues.Remove(InCommandString) > 0;
+		return SavedCommands.Remove(ExistingData) > 0;
 	}
 	
 	return false;
 }
 
-void UConsoleVariablesAsset::CopyFrom(UConsoleVariablesAsset* InAssetToCopy)
+void UConsoleVariablesAsset::CopyFrom(const UConsoleVariablesAsset* InAssetToCopy)
 {
-	VariableCollectionDescription = InAssetToCopy->VariableCollectionDescription;
-	SavedCommandsAndValues = InAssetToCopy->SavedCommandsAndValues;
+	VariableCollectionDescription = InAssetToCopy->GetVariableCollectionDescription();
+	SavedCommands = InAssetToCopy->GetSavedCommands();
 }
diff --git a/Engine/Plugins/Experimental/ConsoleVariables/Source/ConsoleVariablesEditor/Private/ConsoleVariablesEditorModule.cpp b/Engine/Plugins/Experimental/ConsoleVariables/Source/ConsoleVariablesEditor/Private/ConsoleVariablesEditorModule.cpp
--- a/Engine/Plugins/Experimental/ConsoleVariables/Source/ConsoleVariablesEditor/Private/ConsoleVariablesEditorModule.cpp
+++ b/Engine/Plugins/Experimental/ConsoleVariables/Source/ConsoleVariablesEditor/Private/ConsoleVariablesEditorModule.cpp
@@ -183,11 +183,11 @@ void FConsoleVariablesEditorModule::OnConsoleVariableChanged(IConsoleVariable* C
 	if (const TWeakPtr<FConsoleVariablesEditorCommandInfo> CommandInfo =
 		FindCommandInfoByConsoleVariableReference(ChangedVariable); CommandInfo.IsValid())
 	{
-		FString OutValue;
 		const TSharedPtr<FConsoleVariablesEditorCommandInfo>& PinnedCommand = CommandInfo.Pin();
 		const FString& Key = PinnedCommand->Command;
 		
-		const bool bIsVariableCurrentlyTracked = EditingAsset->FindSavedValueByCommandString(Key, OutValue);
+		FConsoleVariablesEditorAssetSaveData FoundData;
+		const bool bIsVariableCurrentlyTracked = EditingAsset->FindSavedDataByCommandString(Key, FoundData);
 		
 		if (!bIsVariableCurrentlyTracked)
 		{
@@ -196,10 +196,9 @@ void FConsoleVariablesEditorModule::OnConsoleVariableChanged(IConsoleVariable* C
 			if (GetMutableDefault<UConsoleVariablesEditorProjectSettings>()->bAddAllChangedConsoleVariablesToCurrentPreset
 				&& PinnedCommand->IsCurrentValueDifferentFromInputValue(PinnedCommand->StartupValueAsString))
 			{
-				EditingAsset->AddOrSetConsoleVariableSavedValue(Key, ChangedVariable->GetString());
 				if (MainPanel.IsValid())
 				{
-					MainPanel->RebuildList();
+					MainPanel->AddConsoleVariable(Key, ChangedVariable->GetString());
 				}
 
 				SendMultiUserConsoleVariableChange(Key, ChangedVariable->GetString());
diff --git a/Engine/Plugins/Experimental/ConsoleVariables/Source/ConsoleVariablesEditor/Private/Views/List/ConsoleVariablesEditorListRow.cpp b/Engine/Plugins/Experimental/ConsoleVariables/Source/ConsoleVariablesEditor/Private/Views/List/ConsoleVariablesEditorListRow.cpp
--- a/Engine/Plugins/Experimental/ConsoleVariables/Source/ConsoleVariablesEditor/Private/Views/List/ConsoleVariablesEditorListRow.cpp
+++ b/Engine/Plugins/Experimental/ConsoleVariables/Source/ConsoleVariablesEditor/Private/Views/List/ConsoleVariablesEditorListRow.cpp
@@ -218,6 +218,16 @@ EVisibility FConsoleVariablesEditorListRow::GetDesiredVisibility() const
 		EVisibility::Visible : EVisibility::Collapsed;
 }
 
+const FString& FConsoleVariablesEditorListRow::GetCachedValue() const
+{
+	return CachedValue;
+}
+
+void FConsoleVariablesEditorListRow::SetCachedValue(const FString& NewCachedValue)
+{
+	CachedValue = NewCachedValue;
+}
+
 TArray<FConsoleVariablesEditorListRowPtr> FConsoleVariablesEditorListRow::GetSelectedTreeViewItems() const
 {
 	return ListViewPtr.Pin()->GetSelectedTreeViewItems();
diff --git a/Engine/Plugins/Experimental/ConsoleVariables/Source/ConsoleVariablesEditor/Private/Views/List/ConsoleVariablesEditorListRow.h b/Engine/Plugins/Experimental/ConsoleVariables/Source/ConsoleVariablesEditor/Private/Views/List/ConsoleVariablesEditorListRow.h
--- a/Engine/Plugins/Experimental/ConsoleVariables/Source/ConsoleVariablesEditor/Private/Views/List/ConsoleVariablesEditorListRow.h
+++ b/Engine/Plugins/Experimental/ConsoleVariables/Source/ConsoleVariablesEditor/Private/Views/List/ConsoleVariablesEditorListRow.h
@@ -6,6 +6,7 @@
 #include "ConsoleVariablesEditorCommandInfo.h"
 #include "Widgets/Input/SCheckBox.h"
 
+class SConsoleVariablesEditorListValueInput;
 class SConsoleVariablesEditorList;
 
 struct FConsoleVariablesEditorListRow;
@@ -101,6 +102,9 @@ struct FConsoleVariablesEditorListRow final : TSharedFromThis<FConsoleVariablesE
 	{
 		return ListViewPtr;
 	}
+	
+	[[nodiscard]] const FString& GetCachedValue() const;
+	void SetCachedValue(const FString& CachedValue);
 
 	[[nodiscard]] TArray<FConsoleVariablesEditorListRowPtr> GetSelectedTreeViewItems() const;
 
@@ -118,6 +122,7 @@ private:
 	ECheckBoxState WidgetCheckedState = ECheckBoxState::Checked;
 	
 	TWeakPtr<SConsoleVariablesEditorList> ListViewPtr;
+	FString CachedValue = "";
 	
 	bool bIsTreeViewItemExpanded = false;
 	bool bShouldFlashOnScrollIntoView = false;
diff --git a/Engine/Plugins/Experimental/ConsoleVariables/Source/ConsoleVariablesEditor/Private/Views/List/SConsoleVariablesEditorList.cpp b/Engine/Plugins/Experimental/ConsoleVariables/Source/ConsoleVariablesEditor/Private/Views/List/SConsoleVariablesEditorList.cpp
--- a/Engine/Plugins/Experimental/ConsoleVariables/Source/ConsoleVariablesEditor/Private/Views/List/SConsoleVariablesEditorList.cpp
+++ b/Engine/Plugins/Experimental/ConsoleVariables/Source/ConsoleVariablesEditor/Private/Views/List/SConsoleVariablesEditorList.cpp
@@ -298,9 +298,9 @@ void SConsoleVariablesEditorList::SetTreeViewItems(const TArray<FConsoleVariable
 	TreeViewPtr->RequestListRefresh();
 }
 
-void SConsoleVariablesEditorList::UpdatePresetValuesForSave(TObjectPtr<UConsoleVariablesAsset> InAsset) const
+void SConsoleVariablesEditorList::UpdatePresetValuesForSave(const TObjectPtr<UConsoleVariablesAsset> InAsset) const
 {
-	TMap<FString, FString> NewSavedValueMap;
+	TArray<FConsoleVariablesEditorAssetSaveData> NewSavedCommands;
 	
 	for (const FConsoleVariablesEditorListRowPtr& Item : TreeViewRootObjects)
 	{
@@ -308,14 +308,17 @@ void SConsoleVariablesEditorList::UpdatePresetValuesForSave(TObjectPtr<UConsoleV
 		
 		if (CommandInfo.IsValid())
 		{
-			if (const TObjectPtr<IConsoleVariable> Variable = CommandInfo.Pin()->ConsoleVariablePtr)
-			{
-				NewSavedValueMap.Add(CommandInfo.Pin()->Command, Variable->GetString());
-			}
+			NewSavedCommands.Add(
+				{
+					CommandInfo.Pin()->Command,
+					Item->GetCachedValue(),
+					Item->GetWidgetCheckedState()
+				}
+			);
 		}
 	}
 
-	InAsset->ReplaceSavedCommandsAndValues(NewSavedValueMap);
+	InAsset->ReplaceSavedCommands(NewSavedCommands);
 }
 
 FString SConsoleVariablesEditorList::GetSearchStringFromSearchInputField() const
@@ -338,25 +341,31 @@ void SConsoleVariablesEditorList::GenerateTreeView()
 	const TObjectPtr<UConsoleVariablesAsset> EditableAsset = ConsoleVariablesEditorModule.GetEditingAsset();
 	check(EditableAsset);
 
-	for (const TPair<FString, FString>& CommandAndValue : EditableAsset->GetSavedCommandsAndValues())
+	for (const FConsoleVariablesEditorAssetSaveData& SavedCommand : EditableAsset->GetSavedCommands())
 	{
-		TWeakPtr<FConsoleVariablesEditorCommandInfo> CommandInfo = ConsoleVariablesEditorModule.FindCommandInfoByName(CommandAndValue.Key);
-		
-		if (CommandInfo.IsValid())
+		// Get corresponding CommandInfo for tracking
+		if (TWeakPtr<FConsoleVariablesEditorCommandInfo> CommandInfo =
+			ConsoleVariablesEditorModule.FindCommandInfoByName(SavedCommand.CommandName); CommandInfo.IsValid())
 		{
-			if (const TObjectPtr<IConsoleVariable> VariablePtr = CommandInfo.Pin()->ConsoleVariablePtr)
-			{
-				if (!VariablePtr->GetString().Equals(CommandAndValue.Value))
-				{
-					CommandInfo.Pin()->ExecuteCommand(CommandAndValue.Value);
-				}
+			const ECheckBoxState NewCheckedState =
+				SavedCommand.CheckedState == ECheckBoxState::Unchecked ?
+						ECheckBoxState::Unchecked : ECheckBoxState::Checked;
 			
-				FConsoleVariablesEditorListRowPtr NewRow = 
-					MakeShared<FConsoleVariablesEditorListRow>(
-							CommandInfo.Pin(), CommandAndValue.Value, FConsoleVariablesEditorListRow::SingleCommand, 
-							ECheckBoxState::Checked, SharedThis(this), TreeViewRootObjects.Num(), nullptr);
-				TreeViewRootObjects.Add(NewRow);
+			// If the row is checked and the saved value differs from the current value,
+			// execute the command with the saved value 
+			if ( NewCheckedState == ECheckBoxState::Checked &&
+				CommandInfo.Pin()->IsCurrentValueDifferentFromInputValue(SavedCommand.CommandValueAsString))
+			{
+				CommandInfo.Pin()->ExecuteCommand(SavedCommand.CommandValueAsString);
 			}
+		
+			FConsoleVariablesEditorListRowPtr NewRow = 
+				MakeShared<FConsoleVariablesEditorListRow>(
+						CommandInfo.Pin(), SavedCommand.CommandValueAsString,
+						FConsoleVariablesEditorListRow::SingleCommand, 
+						NewCheckedState, SharedThis(this), TreeViewRootObjects.Num(), nullptr);
+		
+			TreeViewRootObjects.Add(NewRow);
 		}
 	}
 
diff --git a/Engine/Plugins/Experimental/ConsoleVariables/Source/ConsoleVariablesEditor/Private/Views/List/SConsoleVariablesEditorList.h b/Engine/Plugins/Experimental/ConsoleVariables/Source/ConsoleVariablesEditor/Private/Views/List/SConsoleVariablesEditorList.h
--- a/Engine/Plugins/Experimental/ConsoleVariables/Source/ConsoleVariablesEditor/Private/Views/List/SConsoleVariablesEditorList.h
+++ b/Engine/Plugins/Experimental/ConsoleVariables/Source/ConsoleVariablesEditor/Private/Views/List/SConsoleVariablesEditorList.h
@@ -49,7 +49,7 @@ public:
 	}
 
 	/** Updates the saved values in a UConsoleVariablesAsset so that the command/value map can be saved to disk */
-	void UpdatePresetValuesForSave(TObjectPtr<UConsoleVariablesAsset> InAsset) const;
+	void UpdatePresetValuesForSave(const TObjectPtr<UConsoleVariablesAsset> InAsset) const;
 
 	FString GetSearchStringFromSearchInputField() const;
 	void ExecuteListViewSearchOnAllRows(const FString& SearchString) const;
diff --git a/Engine/Plugins/Experimental/ConsoleVariables/Source/ConsoleVariablesEditor/Private/Views/List/SConsoleVariablesEditorListRow.cpp b/Engine/Plugins/Experimental/ConsoleVariables/Source/ConsoleVariablesEditor/Private/Views/List/SConsoleVariablesEditorListRow.cpp
--- a/Engine/Plugins/Experimental/ConsoleVariables/Source/ConsoleVariablesEditor/Private/Views/List/SConsoleVariablesEditorListRow.cpp
+++ b/Engine/Plugins/Experimental/ConsoleVariables/Source/ConsoleVariablesEditor/Private/Views/List/SConsoleVariablesEditorListRow.cpp
@@ -419,14 +419,14 @@ void SConsoleVariablesEditorListRow::OnCheckboxStateChange(const ECheckBoxState
 	{
 		PinnedItem->SetWidgetCheckedState(InNewState, true);
 		
-		if (PinnedItem->GetRowType() == FConsoleVariablesEditorListRow::SingleCommand && ValueChildInputWidget.IsValid())
+		if (PinnedItem->GetRowType() == FConsoleVariablesEditorListRow::SingleCommand)
 		{									
 			if (PinnedItem->IsRowChecked())
 			{
-				PinnedItem->GetCommandInfo().Pin()->ExecuteCommand(ValueChildInputWidget->GetCachedValue());
+				PinnedItem->GetCommandInfo().Pin()->ExecuteCommand(PinnedItem->GetCachedValue());
 
 				FConsoleVariablesEditorModule::Get().SendMultiUserConsoleVariableChange(
-					PinnedItem->GetCommandInfo().Pin()->Command, ValueChildInputWidget->GetCachedValue());
+					PinnedItem->GetCommandInfo().Pin()->Command, PinnedItem->GetCachedValue());
 			}
 			else
 			{
@@ -458,7 +458,7 @@ TSharedRef<SWidget> SConsoleVariablesEditorListRow::GenerateValueCellWidget(cons
 				{
 					return FText::Format(
 								ValueWidgetToolTipFormatText,
-								FText::FromString(ValueChildInputWidget->GetCachedValue()),
+								FText::FromString(PinnedItem->GetCachedValue()),
 								FText::FromString(PinnedItem->GetPresetValue()),
 								FText::FromString(PinnedCommand->StartupValueAsString),
 								FConsoleVariablesEditorCommandInfo::ConvertConsoleVariableSetByFlagToText(PinnedCommand->StartupSource)
diff --git a/Engine/Plugins/Experimental/ConsoleVariables/Source/ConsoleVariablesEditor/Private/Views/List/SConsoleVariablesEditorListValueInput.cpp b/Engine/Plugins/Experimental/ConsoleVariables/Source/ConsoleVariablesEditor/Private/Views/List/SConsoleVariablesEditorListValueInput.cpp
--- a/Engine/Plugins/Experimental/ConsoleVariables/Source/ConsoleVariablesEditor/Private/Views/List/SConsoleVariablesEditorListValueInput.cpp
+++ b/Engine/Plugins/Experimental/ConsoleVariables/Source/ConsoleVariablesEditor/Private/Views/List/SConsoleVariablesEditorListValueInput.cpp
@@ -33,6 +33,11 @@ TSharedRef<SConsoleVariablesEditorListValueInput> SConsoleVariablesEditorListVal
 				return SNew(SConsoleVariablesEditorListValueInput_Float, InRow);
 			}
 
+			if (Variable->IsVariableBool())
+			{
+				return SNew(SConsoleVariablesEditorListValueInput_Bool, InRow);
+			}
+
 			if (Variable->IsVariableInt())
 			{
 				return SNew(SConsoleVariablesEditorListValueInput_Int, InRow);
@@ -45,18 +50,8 @@ TSharedRef<SConsoleVariablesEditorListValueInput> SConsoleVariablesEditorListVal
 		}
 	}
 
-	// bool
-	return SNew(SConsoleVariablesEditorListValueInput_Bool, InRow);
-}
-
-const FString& SConsoleVariablesEditorListValueInput::GetCachedValue() const
-{
-	return CachedValue;
-}
-
-void SConsoleVariablesEditorListValueInput::SetCachedValue(const FString& NewCachedValue)
-{
-	CachedValue = NewCachedValue;
+	// fallback
+	return SNew(SConsoleVariablesEditorListValueInput_String, InRow);
 }
 
 bool SConsoleVariablesEditorListValueInput::IsRowChecked() const
@@ -77,32 +72,35 @@ void SConsoleVariablesEditorListValueInput_Float::Construct(const FArguments& In
 		.MaxFractionalDigits(3)
 		.Value_Lambda([this]
 		{
-			if (Item.IsValid() &&
-				(Item.Pin()->GetWidgetCheckedState() == ECheckBoxState::Checked ||
+			check (Item.IsValid());
+			
+			if ((Item.Pin()->GetWidgetCheckedState() == ECheckBoxState::Checked ||
 					(ProjectSettingsPtr && ProjectSettingsPtr->UncheckedRowDisplayType == EConsoleVariablesEditorRowDisplayType::ShowCurrentValue)))
 			{
 				return FCString::Atof(*Item.Pin()->GetCommandInfo().Pin()->ConsoleVariablePtr->GetString());
 			}
 
-			return FCString::Atof(*GetCachedValue());
+			return FCString::Atof(*Item.Pin()->GetCachedValue());
 		})
 		.OnValueChanged_Lambda([this] (const float InValue)
 		{
+			check (Item.IsValid());
+			
 			const FString ValueAsString = FString::SanitizeFloat(InValue);
 			
-			if (const TSharedPtr<FConsoleVariablesEditorListRow> PinnedItem = Item.Pin(); !GetCachedValue().Equals(ValueAsString))
+			if (const TSharedPtr<FConsoleVariablesEditorListRow> PinnedItem = Item.Pin(); !PinnedItem->GetCachedValue().Equals(ValueAsString))
 			{				
 				PinnedItem->GetCommandInfo().Pin()->ExecuteCommand(ValueAsString);
 
 				FConsoleVariablesEditorModule::Get().SendMultiUserConsoleVariableChange(PinnedItem->GetCommandInfo().Pin()->Command, ValueAsString);
-			}
 			
-			SetCachedValue(ValueAsString);
+				PinnedItem->SetCachedValue(ValueAsString);
+			}
 		})
 		.IsEnabled(this, &SConsoleVariablesEditorListValueInput::IsRowChecked)
 	];
 
-	SetCachedValue(GetInputValueAsString());
+	Item.Pin()->SetCachedValue(GetInputValueAsString());
 }
 
 void SConsoleVariablesEditorListValueInput_Float::SetInputValue(const FString& InValueAsString)
@@ -132,37 +130,56 @@ void SConsoleVariablesEditorListValueInput_Int::Construct(const FArguments& InAr
 		SAssignNew(InputWidget, SSpinBox<int32>)
 		.Value_Lambda([this]
 		{
-			if (Item.IsValid() &&
-				(Item.Pin()->GetWidgetCheckedState() == ECheckBoxState::Checked ||
+			check (Item.IsValid());
+			
+			if ((Item.Pin()->GetWidgetCheckedState() == ECheckBoxState::Checked ||
 					(ProjectSettingsPtr && ProjectSettingsPtr->UncheckedRowDisplayType == EConsoleVariablesEditorRowDisplayType::ShowCurrentValue)))
 			{
 				return FCString::Atoi(*Item.Pin()->GetCommandInfo().Pin()->ConsoleVariablePtr->GetString());
 			}
 
-			return FCString::Atoi(*GetCachedValue());
+			return FCString::Atoi(*Item.Pin()->GetCachedValue());
 		})
 		.OnValueChanged_Lambda([this] (const int32 InValue)
 		{
+			check (Item.IsValid());
+			
 			const FString ValueAsString = FString::FromInt(InValue);
 			
-			if (const TSharedPtr<FConsoleVariablesEditorListRow> PinnedItem = Item.Pin(); !GetCachedValue().Equals(ValueAsString))
+			if (const TSharedPtr<FConsoleVariablesEditorListRow> PinnedItem = Item.Pin(); !PinnedItem->GetCachedValue().Equals(ValueAsString))
 			{				
 				PinnedItem->GetCommandInfo().Pin()->ExecuteCommand(ValueAsString);
 
 				FConsoleVariablesEditorModule::Get().SendMultiUserConsoleVariableChange(PinnedItem->GetCommandInfo().Pin()->Command, ValueAsString);
-			}
 
-			SetCachedValue(ValueAsString);
+				PinnedItem->SetCachedValue(ValueAsString);
+			}
 		})
 		.IsEnabled(this, &SConsoleVariablesEditorListValueInput::IsRowChecked)
 	];
 	
-	SetCachedValue(GetInputValueAsString());
+	Item.Pin()->SetCachedValue(GetInputValueAsString());
 }
 
 void SConsoleVariablesEditorListValueInput_Int::SetInputValue(const FString& InValueAsString)
 {
-	InputWidget->SetValue(FCString::Atoi(*InValueAsString));
+	if (InValueAsString.IsNumeric())
+	{
+		InputWidget->SetValue(FCString::Atoi(*InValueAsString));
+	}
+	else
+	{
+		InputWidget->SetValue(2);
+		
+		if (InValueAsString.TrimStartAndEnd().ToLower() == "true")
+		{
+			InputWidget->SetValue(1);
+		}
+		else if (InValueAsString.TrimStartAndEnd().ToLower() == "false")
+		{
+			InputWidget->SetValue(0);
+		}
+	}
 }
 
 FString SConsoleVariablesEditorListValueInput_Int::GetInputValueAsString()
@@ -187,32 +204,35 @@ void SConsoleVariablesEditorListValueInput_String::Construct(const FArguments& I
 		SAssignNew(InputWidget, SEditableText)
 		.Text_Lambda([this]
 		{
-			if (Item.IsValid() &&
-				(Item.Pin()->GetWidgetCheckedState() == ECheckBoxState::Checked ||
+			check (Item.IsValid());
+			
+			if ((Item.Pin()->GetWidgetCheckedState() == ECheckBoxState::Checked ||
 					(ProjectSettingsPtr && ProjectSettingsPtr->UncheckedRowDisplayType == EConsoleVariablesEditorRowDisplayType::ShowCurrentValue)))
 			{
 				return FText::FromString(*Item.Pin()->GetCommandInfo().Pin()->ConsoleVariablePtr->GetString());
 			}
 
-			return FText::FromString(GetCachedValue());
+			return FText::FromString(Item.Pin()->GetCachedValue());
 		})
 		.OnTextCommitted_Lambda([this] (const FText& InValue, ETextCommit::Type InTextCommitType)
 		{
+			check (Item.IsValid());
+			
 			const FString ValueAsString = InValue.ToString();
 			
-			if (const TSharedPtr<FConsoleVariablesEditorListRow> PinnedItem = Item.Pin(); !GetCachedValue().Equals(ValueAsString))
+			if (const TSharedPtr<FConsoleVariablesEditorListRow> PinnedItem = Item.Pin(); !PinnedItem->GetCachedValue().Equals(ValueAsString))
 			{				
 				PinnedItem->GetCommandInfo().Pin()->ExecuteCommand(ValueAsString);
 
 				FConsoleVariablesEditorModule::Get().SendMultiUserConsoleVariableChange(PinnedItem->GetCommandInfo().Pin()->Command, ValueAsString);
-			}
 
-			SetCachedValue(ValueAsString);
+				PinnedItem->SetCachedValue(ValueAsString);
+			}
 		})
 		.IsEnabled(this, &SConsoleVariablesEditorListValueInput::IsRowChecked)
 	];
 
-	SetCachedValue(GetInputValueAsString());
+	Item.Pin()->SetCachedValue(GetInputValueAsString());
 }
 
 void SConsoleVariablesEditorListValueInput_String::SetInputValue(const FString& InValueAsString)
@@ -242,54 +262,41 @@ void SConsoleVariablesEditorListValueInput_Bool::Construct(const FArguments& InA
 		SAssignNew(InputWidget, SSpinBox<int32>)
 		.Value_Lambda([this]
 		{
-			if (Item.IsValid() &&
-				(Item.Pin()->GetWidgetCheckedState() == ECheckBoxState::Checked ||
-					(ProjectSettingsPtr && ProjectSettingsPtr->UncheckedRowDisplayType == EConsoleVariablesEditorRowDisplayType::ShowCurrentValue)))
+			check (Item.IsValid());
+			
+			if (Item.Pin()->GetWidgetCheckedState() == ECheckBoxState::Checked ||
+					(ProjectSettingsPtr && ProjectSettingsPtr->UncheckedRowDisplayType == EConsoleVariablesEditorRowDisplayType::ShowCurrentValue))
 			{
 				return FCString::Atoi(*Item.Pin()->GetCommandInfo().Pin()->ConsoleVariablePtr->GetString());
 			}
 
-			return FCString::Atoi(*GetCachedValue());
+			return FCString::Atoi(*Item.Pin()->GetCachedValue());
 		})
 		.MinSliderValue(0)
 		.MaxSliderValue(2)
 		.OnValueChanged_Lambda([this] (const int32 InValue)
 		{
+			check (Item.IsValid());
+			
 			const FString ValueAsString = FString::FromInt(InValue);
 			
-			if (const TSharedPtr<FConsoleVariablesEditorListRow> PinnedItem = Item.Pin(); !GetCachedValue().Equals(ValueAsString))
+			if (const TSharedPtr<FConsoleVariablesEditorListRow> PinnedItem = Item.Pin(); !PinnedItem->GetCachedValue().Equals(ValueAsString))
 			{				
 				PinnedItem->GetCommandInfo().Pin()->ExecuteCommand(ValueAsString);
 
 				FConsoleVariablesEditorModule::Get().SendMultiUserConsoleVariableChange(PinnedItem->GetCommandInfo().Pin()->Command, ValueAsString);
-			}
 
-			SetCachedValue(ValueAsString);
+				PinnedItem->SetCachedValue(ValueAsString);
+			}
 		})
 		.IsEnabled(this, &SConsoleVariablesEditorListValueInput::IsRowChecked)
 	];
-	SetCachedValue(GetInputValueAsString());
+	Item.Pin()->SetCachedValue(GetInputValueAsString());
 }
 
 void SConsoleVariablesEditorListValueInput_Bool::SetInputValue(const FString& InValueAsString)
 {
-	if (InValueAsString.IsNumeric())
-	{
-		InputWidget->SetValue(FMath::Clamp(FCString::Atoi(*InValueAsString), 0, 2));
-	}
-	else
-	{
-		InputWidget->SetValue(2);
-		
-		if (InValueAsString.TrimStartAndEnd().ToLower() == "true")
-		{
-			InputWidget->SetValue(1);
-		}
-		else if (InValueAsString.TrimStartAndEnd().ToLower() == "false")
-		{
-			InputWidget->SetValue(0);
-		}
-	}
+	
 }
 
 FString SConsoleVariablesEditorListValueInput_Bool::GetInputValueAsString()
@@ -312,4 +319,25 @@ FString SConsoleVariablesEditorListValueInput_Bool::GetBoolValueAsString() const
 	return GetInputValueAsBool() ? "true" : "false";
 }
 
+int32 SConsoleVariablesEditorListValueInput_Bool::ConvertBoolValueToInteger(const FString& InValue) const
+{
+	if (InValue.IsNumeric())
+	{
+		return FMath::Clamp(FCString::Atoi(*InValue), 0, 2);
+	}
+	else
+	{		
+		if (InValue.TrimStartAndEnd().ToLower() == "true")
+		{
+			return 1;
+		}
+		else if (InValue.TrimStartAndEnd().ToLower() == "false")
+		{
+			return 0;
+		}
+
+		return 2;
+	}
+}
+
 #undef LOCTEXT_NAMESPACE
diff --git a/Engine/Plugins/Experimental/ConsoleVariables/Source/ConsoleVariablesEditor/Private/Views/List/SConsoleVariablesEditorListValueInput.h b/Engine/Plugins/Experimental/ConsoleVariables/Source/ConsoleVariablesEditor/Private/Views/List/SConsoleVariablesEditorListValueInput.h
--- a/Engine/Plugins/Experimental/ConsoleVariables/Source/ConsoleVariablesEditor/Private/Views/List/SConsoleVariablesEditorListValueInput.h
+++ b/Engine/Plugins/Experimental/ConsoleVariables/Source/ConsoleVariablesEditor/Private/Views/List/SConsoleVariablesEditorListValueInput.h
@@ -18,16 +18,12 @@ public:
 	
 	virtual void SetInputValue(const FString& InValueAsString) = 0;
 	virtual FString GetInputValueAsString() = 0;
-	
-	[[nodiscard]] const FString& GetCachedValue() const;
-	void SetCachedValue(const FString& CachedValue);
 
 	bool IsRowChecked() const;
 
 protected:
 	
 	TWeakPtr<FConsoleVariablesEditorListRow> Item = nullptr;
-	FString CachedValue = "";
 
 	TObjectPtr<UConsoleVariablesEditorProjectSettings> ProjectSettingsPtr = nullptr;
 };
@@ -114,6 +110,8 @@ public:
 
 	FString GetBoolValueAsString() const;
 
+	int32 ConvertBoolValueToInteger(const FString& InValue) const;
+
 private:
 	
 	TSharedPtr<SSpinBox<int32>> InputWidget;
diff --git a/Engine/Plugins/Experimental/ConsoleVariables/Source/ConsoleVariablesEditor/Private/Views/MainPanel/ConsoleVariablesEditorMainPanel.cpp b/Engine/Plugins/Experimental/ConsoleVariables/Source/ConsoleVariablesEditor/Private/Views/MainPanel/ConsoleVariablesEditorMainPanel.cpp
--- a/Engine/Plugins/Experimental/ConsoleVariables/Source/ConsoleVariablesEditor/Private/Views/MainPanel/ConsoleVariablesEditorMainPanel.cpp
+++ b/Engine/Plugins/Experimental/ConsoleVariables/Source/ConsoleVariablesEditor/Private/Views/MainPanel/ConsoleVariablesEditorMainPanel.cpp
@@ -60,7 +60,13 @@ void FConsoleVariablesEditorMainPanel::AddConsoleVariable(
 	{
 		UConsoleVariablesAsset* Asset = EditingAsset.Get();
 
-		Asset->AddOrSetConsoleVariableSavedValue(InConsoleCommand, InValue);
+		Asset->AddOrSetConsoleVariableSavedData(
+			{
+				InConsoleCommand,
+				InValue,
+				ECheckBoxState::Checked
+			}
+		);
 
 		RebuildList(bScrollToNewRow ? InConsoleCommand : "");
 	}
diff --git a/Engine/Plugins/Experimental/ConsoleVariables/Source/ConsoleVariablesEditor/Public/ConsoleVariablesAsset.h b/Engine/Plugins/Experimental/ConsoleVariables/Source/ConsoleVariablesEditor/Public/ConsoleVariablesAsset.h
--- a/Engine/Plugins/Experimental/ConsoleVariables/Source/ConsoleVariablesEditor/Public/ConsoleVariablesAsset.h
+++ b/Engine/Plugins/Experimental/ConsoleVariables/Source/ConsoleVariablesEditor/Public/ConsoleVariablesAsset.h
@@ -8,6 +8,28 @@
 
 struct FConsoleVariablesEditorCommandInfo;
 
+/** Data that will be serialized with this asset */
+USTRUCT()
+struct FConsoleVariablesEditorAssetSaveData
+{
+	GENERATED_BODY()
+
+	FORCEINLINE bool operator==(const FConsoleVariablesEditorAssetSaveData& Comparator) const
+	{
+		return CommandName.Equals(Comparator.CommandName);
+	}
+	
+	UPROPERTY()
+	FString CommandName;
+
+	UPROPERTY()
+	FString CommandValueAsString;
+
+	UPROPERTY()
+	// If Undetermined, we can assume this data was not previously saved
+	ECheckBoxState CheckedState = ECheckBoxState::Undetermined;
+};
+
 /** An asset used to track collections of console variables that can be recalled and edited using the Console Variables Editor. */
 UCLASS(BlueprintType)
 class CONSOLEVARIABLESEDITOR_API UConsoleVariablesAsset : public UObject
@@ -25,30 +47,31 @@ public:
 	}
 
 	/** Returns the saved list of console variable information such as the variable name, the type and the value of the vriable at the time the asset was saved. */
-	const TMap<FString, FString>& GetSavedCommandsAndValues() const
+	const TArray<FConsoleVariablesEditorAssetSaveData>& GetSavedCommands() const
 	{
-		return SavedCommandsAndValues;
+		return SavedCommands;
 	}
 
-	void ReplaceSavedCommandsAndValues(const TMap<FString, FString>& InMap);
+	/** Completely replaces the saved data with new saved data */
+	void ReplaceSavedCommands(const TArray<FConsoleVariablesEditorAssetSaveData>& Replacement);
 
 	/** Returns how many console variables are serialized in this asset */
-	int32 GetSavedCommandsAndValuesCount() const
+	int32 GetSavedCommandsCount() const
 	{
-		return GetSavedCommandsAndValues().Num();
+		return GetSavedCommands().Num();
 	}
 
 	/** Outputs the FConsoleVariablesEditorCommandInfo matching InCommand. Returns whether a match was found. Case sensitive. */
-	bool FindSavedValueByCommandString(const FString& InCommandString, FString& OutValue) const;
+	bool FindSavedDataByCommandString(const FString& InCommandString, FConsoleVariablesEditorAssetSaveData& OutValue) const;
 
 	/** Set the value of a saved console variable if the name matches; add a new console variable to the list if a match is not found. */
-	void AddOrSetConsoleVariableSavedValue(const FString& InCommandString, const FString& InNewValue);
+	void AddOrSetConsoleVariableSavedData(const FConsoleVariablesEditorAssetSaveData& InData);
 
 	/** Returns true if the element was found and successfully removed. */
 	bool RemoveConsoleVariable(const FString& InCommandString);
 
 	/** Copy data from input asset to this asset */
-	void CopyFrom(UConsoleVariablesAsset* InAssetToCopy);
+	void CopyFrom(const UConsoleVariablesAsset* InAssetToCopy);
 	
 private:
 	
@@ -58,5 +81,5 @@ private:
 
 	/** A saved list of console variable information such as the variable name, the type and the value of the variable at the time the asset was saved. */
 	UPROPERTY()
-	TMap<FString, FString> SavedCommandsAndValues;
+	TArray<FConsoleVariablesEditorAssetSaveData> SavedCommands;
 };
