Console Variables Editor:

-Change Bool Types to button
-Check for bool type before int type, as bools are considered ints in cvars but ints aren't considered bools
-Fixed an issue in which resetting to preset value wouldn't update the cached value

#[fyi] Jason.Walter

diff --git a/Engine/Plugins/Experimental/ConsoleVariables/Source/ConsoleVariablesEditor/Private/Views/List/ConsoleVariablesEditorListRow.cpp b/Engine/Plugins/Experimental/ConsoleVariables/Source/ConsoleVariablesEditor/Private/Views/List/ConsoleVariablesEditorListRow.cpp
--- a/Engine/Plugins/Experimental/ConsoleVariables/Source/ConsoleVariablesEditor/Private/Views/List/ConsoleVariablesEditorListRow.cpp
+++ b/Engine/Plugins/Experimental/ConsoleVariables/Source/ConsoleVariablesEditor/Private/Views/List/ConsoleVariablesEditorListRow.cpp
@@ -254,9 +254,11 @@ FReply FConsoleVariablesEditorListRow::OnRemoveButtonClicked()
 	return FReply::Handled();
 }
 
-void FConsoleVariablesEditorListRow::ResetToPresetValue() const
+void FConsoleVariablesEditorListRow::ResetToPresetValue()
 {
-	GetCommandInfo().Pin()->ExecuteCommand(GetPresetValue());
+	const FString& Value = GetPresetValue();
+	GetCommandInfo().Pin()->ExecuteCommand(Value);
+	SetCachedValue(Value);
 }
 
 bool FConsoleVariablesEditorListRow::GetShouldFlashOnScrollIntoView() const
diff --git a/Engine/Plugins/Experimental/ConsoleVariables/Source/ConsoleVariablesEditor/Private/Views/List/ConsoleVariablesEditorListRow.h b/Engine/Plugins/Experimental/ConsoleVariables/Source/ConsoleVariablesEditor/Private/Views/List/ConsoleVariablesEditorListRow.h
--- a/Engine/Plugins/Experimental/ConsoleVariables/Source/ConsoleVariablesEditor/Private/Views/List/ConsoleVariablesEditorListRow.h
+++ b/Engine/Plugins/Experimental/ConsoleVariables/Source/ConsoleVariablesEditor/Private/Views/List/ConsoleVariablesEditorListRow.h
@@ -110,7 +110,7 @@ struct FConsoleVariablesEditorListRow final : TSharedFromThis<FConsoleVariablesE
 
 	FReply OnRemoveButtonClicked();
 	
-	void ResetToPresetValue() const;
+	void ResetToPresetValue();
 
 private:
 
diff --git a/Engine/Plugins/Experimental/ConsoleVariables/Source/ConsoleVariablesEditor/Private/Views/List/SConsoleVariablesEditorListValueInput.cpp b/Engine/Plugins/Experimental/ConsoleVariables/Source/ConsoleVariablesEditor/Private/Views/List/SConsoleVariablesEditorListValueInput.cpp
--- a/Engine/Plugins/Experimental/ConsoleVariables/Source/ConsoleVariablesEditor/Private/Views/List/SConsoleVariablesEditorListValueInput.cpp
+++ b/Engine/Plugins/Experimental/ConsoleVariables/Source/ConsoleVariablesEditor/Private/Views/List/SConsoleVariablesEditorListValueInput.cpp
@@ -62,6 +62,8 @@ bool SConsoleVariablesEditorListValueInput::IsRowChecked() const
 void SConsoleVariablesEditorListValueInput_Float::Construct(const FArguments& InArgs,
                                                             const TWeakPtr<FConsoleVariablesEditorListRow> InRow)
 {
+	check (InRow.IsValid());
+	
 	Item = InRow;
 
 	ProjectSettingsPtr = GetMutableDefault<UConsoleVariablesEditorProjectSettings>();
@@ -74,8 +76,9 @@ void SConsoleVariablesEditorListValueInput_Float::Construct(const FArguments& In
 		{
 			check (Item.IsValid());
 			
-			if ((Item.Pin()->GetWidgetCheckedState() == ECheckBoxState::Checked ||
-					(ProjectSettingsPtr && ProjectSettingsPtr->UncheckedRowDisplayType == EConsoleVariablesEditorRowDisplayType::ShowCurrentValue)))
+			if (Item.Pin()->GetWidgetCheckedState() == ECheckBoxState::Checked ||
+					(ProjectSettingsPtr &&
+						ProjectSettingsPtr->UncheckedRowDisplayType == EConsoleVariablesEditorRowDisplayType::ShowCurrentValue))
 			{
 				return FCString::Atof(*Item.Pin()->GetCommandInfo().Pin()->ConsoleVariablePtr->GetString());
 			}
@@ -88,11 +91,13 @@ void SConsoleVariablesEditorListValueInput_Float::Construct(const FArguments& In
 			
 			const FString ValueAsString = FString::SanitizeFloat(InValue);
 			
-			if (const TSharedPtr<FConsoleVariablesEditorListRow> PinnedItem = Item.Pin(); !PinnedItem->GetCachedValue().Equals(ValueAsString))
+			if (const TSharedPtr<FConsoleVariablesEditorListRow> PinnedItem = Item.Pin();
+				!PinnedItem->GetCachedValue().Equals(ValueAsString))
 			{				
 				PinnedItem->GetCommandInfo().Pin()->ExecuteCommand(ValueAsString);
 
-				FConsoleVariablesEditorModule::Get().SendMultiUserConsoleVariableChange(PinnedItem->GetCommandInfo().Pin()->Command, ValueAsString);
+				FConsoleVariablesEditorModule::Get().SendMultiUserConsoleVariableChange(
+					PinnedItem->GetCommandInfo().Pin()->Command, ValueAsString);
 			
 				PinnedItem->SetCachedValue(ValueAsString);
 			}
@@ -121,6 +126,8 @@ float SConsoleVariablesEditorListValueInput_Float::GetInputValue() const
 void SConsoleVariablesEditorListValueInput_Int::Construct(const FArguments& InArgs,
                                                           const TWeakPtr<FConsoleVariablesEditorListRow> InRow)
 {
+	check (InRow.IsValid());
+	
 	Item = InRow;
 
 	ProjectSettingsPtr = GetMutableDefault<UConsoleVariablesEditorProjectSettings>();
@@ -132,8 +139,9 @@ void SConsoleVariablesEditorListValueInput_Int::Construct(const FArguments& InAr
 		{
 			check (Item.IsValid());
 			
-			if ((Item.Pin()->GetWidgetCheckedState() == ECheckBoxState::Checked ||
-					(ProjectSettingsPtr && ProjectSettingsPtr->UncheckedRowDisplayType == EConsoleVariablesEditorRowDisplayType::ShowCurrentValue)))
+			if (Item.Pin()->GetWidgetCheckedState() == ECheckBoxState::Checked ||
+					(ProjectSettingsPtr &&
+						ProjectSettingsPtr->UncheckedRowDisplayType == EConsoleVariablesEditorRowDisplayType::ShowCurrentValue))
 			{
 				return FCString::Atoi(*Item.Pin()->GetCommandInfo().Pin()->ConsoleVariablePtr->GetString());
 			}
@@ -146,11 +154,13 @@ void SConsoleVariablesEditorListValueInput_Int::Construct(const FArguments& InAr
 			
 			const FString ValueAsString = FString::FromInt(InValue);
 			
-			if (const TSharedPtr<FConsoleVariablesEditorListRow> PinnedItem = Item.Pin(); !PinnedItem->GetCachedValue().Equals(ValueAsString))
+			if (const TSharedPtr<FConsoleVariablesEditorListRow> PinnedItem = Item.Pin();
+				!PinnedItem->GetCachedValue().Equals(ValueAsString))
 			{				
 				PinnedItem->GetCommandInfo().Pin()->ExecuteCommand(ValueAsString);
 
-				FConsoleVariablesEditorModule::Get().SendMultiUserConsoleVariableChange(PinnedItem->GetCommandInfo().Pin()->Command, ValueAsString);
+				FConsoleVariablesEditorModule::Get().SendMultiUserConsoleVariableChange(
+					PinnedItem->GetCommandInfo().Pin()->Command, ValueAsString);
 
 				PinnedItem->SetCachedValue(ValueAsString);
 			}
@@ -195,6 +205,8 @@ int32 SConsoleVariablesEditorListValueInput_Int::GetInputValue() const
 void SConsoleVariablesEditorListValueInput_String::Construct(const FArguments& InArgs,
                                                              const TWeakPtr<FConsoleVariablesEditorListRow> InRow)
 {
+	check (InRow.IsValid());
+	
 	Item = InRow;
 
 	ProjectSettingsPtr = GetMutableDefault<UConsoleVariablesEditorProjectSettings>();
@@ -206,8 +218,9 @@ void SConsoleVariablesEditorListValueInput_String::Construct(const FArguments& I
 		{
 			check (Item.IsValid());
 			
-			if ((Item.Pin()->GetWidgetCheckedState() == ECheckBoxState::Checked ||
-					(ProjectSettingsPtr && ProjectSettingsPtr->UncheckedRowDisplayType == EConsoleVariablesEditorRowDisplayType::ShowCurrentValue)))
+			if (Item.Pin()->GetWidgetCheckedState() == ECheckBoxState::Checked ||
+					(ProjectSettingsPtr &&
+						ProjectSettingsPtr->UncheckedRowDisplayType == EConsoleVariablesEditorRowDisplayType::ShowCurrentValue))
 			{
 				return FText::FromString(*Item.Pin()->GetCommandInfo().Pin()->ConsoleVariablePtr->GetString());
 			}
@@ -220,11 +233,13 @@ void SConsoleVariablesEditorListValueInput_String::Construct(const FArguments& I
 			
 			const FString ValueAsString = InValue.ToString();
 			
-			if (const TSharedPtr<FConsoleVariablesEditorListRow> PinnedItem = Item.Pin(); !PinnedItem->GetCachedValue().Equals(ValueAsString))
+			if (const TSharedPtr<FConsoleVariablesEditorListRow> PinnedItem = Item.Pin();
+				!PinnedItem->GetCachedValue().Equals(ValueAsString))
 			{				
 				PinnedItem->GetCommandInfo().Pin()->ExecuteCommand(ValueAsString);
 
-				FConsoleVariablesEditorModule::Get().SendMultiUserConsoleVariableChange(PinnedItem->GetCommandInfo().Pin()->Command, ValueAsString);
+				FConsoleVariablesEditorModule::Get().SendMultiUserConsoleVariableChange(
+					PinnedItem->GetCommandInfo().Pin()->Command, ValueAsString);
 
 				PinnedItem->SetCachedValue(ValueAsString);
 			}
@@ -253,91 +268,74 @@ FString SConsoleVariablesEditorListValueInput_String::GetInputValue() const
 void SConsoleVariablesEditorListValueInput_Bool::Construct(const FArguments& InArgs,
                                                            const TWeakPtr<FConsoleVariablesEditorListRow> InRow)
 {
+	check (InRow.IsValid());
+	
 	Item = InRow;
 
 	ProjectSettingsPtr = GetMutableDefault<UConsoleVariablesEditorProjectSettings>();
 	
 	ChildSlot
 	[
-		SAssignNew(InputWidget, SSpinBox<int32>)
-		.Value_Lambda([this]
-		{
-			check (Item.IsValid());
-			
-			if (Item.Pin()->GetWidgetCheckedState() == ECheckBoxState::Checked ||
-					(ProjectSettingsPtr && ProjectSettingsPtr->UncheckedRowDisplayType == EConsoleVariablesEditorRowDisplayType::ShowCurrentValue))
-			{
-				return FCString::Atoi(*Item.Pin()->GetCommandInfo().Pin()->ConsoleVariablePtr->GetString());
-			}
-
-			return FCString::Atoi(*Item.Pin()->GetCachedValue());
-		})
-		.MinSliderValue(0)
-		.MaxSliderValue(2)
-		.OnValueChanged_Lambda([this] (const int32 InValue)
+		SAssignNew(InputWidget, SButton)
+		.OnClicked_Lambda([this] ()
 		{
-			check (Item.IsValid());
-			
-			const FString ValueAsString = FString::FromInt(InValue);
-			
-			if (const TSharedPtr<FConsoleVariablesEditorListRow> PinnedItem = Item.Pin(); !PinnedItem->GetCachedValue().Equals(ValueAsString))
-			{				
-				PinnedItem->GetCommandInfo().Pin()->ExecuteCommand(ValueAsString);
+			const bool bValueAsBool = GetInputValue();
 
-				FConsoleVariablesEditorModule::Get().SendMultiUserConsoleVariableChange(PinnedItem->GetCommandInfo().Pin()->Command, ValueAsString);
+			SetInputValue(!bValueAsBool);
 
-				PinnedItem->SetCachedValue(ValueAsString);
-			}
+			return FReply::Handled();
 		})
 		.IsEnabled(this, &SConsoleVariablesEditorListValueInput::IsRowChecked)
+		[
+			SAssignNew(ButtonText, STextBlock)
+			.Justification(ETextJustify::Center)
+			.Text_Lambda([this]()
+			{
+				check (Item.IsValid());
+			
+				if (Item.Pin()->GetWidgetCheckedState() == ECheckBoxState::Checked ||
+						(ProjectSettingsPtr &&
+							ProjectSettingsPtr->UncheckedRowDisplayType == EConsoleVariablesEditorRowDisplayType::ShowCurrentValue))
+				{
+					return FText::FromString(Item.Pin()->GetCommandInfo().Pin()->ConsoleVariablePtr->GetString());
+				}
+
+				return FText::FromString(Item.Pin()->GetCachedValue());
+			})
+		]
 	];
+	
 	Item.Pin()->SetCachedValue(GetInputValueAsString());
 }
 
 void SConsoleVariablesEditorListValueInput_Bool::SetInputValue(const FString& InValueAsString)
 {
-	
-}
+	check (Item.IsValid());
+			
+	if (const TSharedPtr<FConsoleVariablesEditorListRow> PinnedItem = Item.Pin();
+		!PinnedItem->GetCachedValue().Equals(InValueAsString))
+	{				
+		PinnedItem->GetCommandInfo().Pin()->ExecuteCommand(InValueAsString);
 
-FString SConsoleVariablesEditorListValueInput_Bool::GetInputValueAsString()
-{
-	return FString::FromInt(GetInputValue());
-}
+		FConsoleVariablesEditorModule::Get().SendMultiUserConsoleVariableChange(PinnedItem->GetCommandInfo().Pin()->Command, InValueAsString);
 
-int32 SConsoleVariablesEditorListValueInput_Bool::GetInputValue() const
-{
-	return InputWidget->GetValue();
+		PinnedItem->SetCachedValue(InValueAsString);
+	}
 }
 
-bool SConsoleVariablesEditorListValueInput_Bool::GetInputValueAsBool() const
-{
-	return GetInputValue() == 1 ? true : false;
+void SConsoleVariablesEditorListValueInput_Bool::SetInputValue(const bool bNewValue)
+{	
+	SetInputValue(BoolToString(bNewValue));
 }
 
-FString SConsoleVariablesEditorListValueInput_Bool::GetBoolValueAsString() const
+FString SConsoleVariablesEditorListValueInput_Bool::GetInputValueAsString()
 {
-	return GetInputValueAsBool() ? "true" : "false";
+	return ButtonText->GetText().ToString();
 }
 
-int32 SConsoleVariablesEditorListValueInput_Bool::ConvertBoolValueToInteger(const FString& InValue) const
+bool SConsoleVariablesEditorListValueInput_Bool::GetInputValue()
 {
-	if (InValue.IsNumeric())
-	{
-		return FMath::Clamp(FCString::Atoi(*InValue), 0, 2);
-	}
-	else
-	{		
-		if (InValue.TrimStartAndEnd().ToLower() == "true")
-		{
-			return 1;
-		}
-		else if (InValue.TrimStartAndEnd().ToLower() == "false")
-		{
-			return 0;
-		}
-
-		return 2;
-	}
+	return StringToBool(GetInputValueAsString());
 }
 
 #undef LOCTEXT_NAMESPACE
diff --git a/Engine/Plugins/Experimental/ConsoleVariables/Source/ConsoleVariablesEditor/Private/Views/List/SConsoleVariablesEditorListValueInput.h b/Engine/Plugins/Experimental/ConsoleVariables/Source/ConsoleVariablesEditor/Private/Views/List/SConsoleVariablesEditorListValueInput.h
--- a/Engine/Plugins/Experimental/ConsoleVariables/Source/ConsoleVariablesEditor/Private/Views/List/SConsoleVariablesEditorListValueInput.h
+++ b/Engine/Plugins/Experimental/ConsoleVariables/Source/ConsoleVariablesEditor/Private/Views/List/SConsoleVariablesEditorListValueInput.h
@@ -2,8 +2,10 @@
 
 #pragma once
 
+#include "Widgets/Input/SButton.h"
 #include "Widgets/Input/SSpinBox.h"
 #include "Widgets/SCompoundWidget.h"
+#include "Widgets/Text/STextBlock.h"
 
 struct FConsoleVariablesEditorListRow;
 
@@ -83,7 +85,7 @@ public:
 	
 	virtual void SetInputValue(const FString& InValueAsString) override;
 	virtual FString GetInputValueAsString() override;
-
+	
 	FString GetInputValue() const;
 
 private:
@@ -103,16 +105,23 @@ public:
 	void Construct(const FArguments& InArgs, const TWeakPtr<FConsoleVariablesEditorListRow> InRow);
 	
 	virtual void SetInputValue(const FString& InValueAsString) override;
+	void SetInputValue(const bool bNewValue);
 	virtual FString GetInputValueAsString() override;
 
-	int32 GetInputValue() const;
-	bool GetInputValueAsBool() const;
+	bool GetInputValue();
 
-	FString GetBoolValueAsString() const;
+	static bool StringToBool(const FString& InString)
+	{
+		return InString.TrimStartAndEnd().ToLower() == "true";
+	}
 
-	int32 ConvertBoolValueToInteger(const FString& InValue) const;
+	static FString BoolToString(const bool bNewBool)
+	{
+		return bNewBool ? "true" : "false";
+	}
 
 private:
 	
-	TSharedPtr<SSpinBox<int32>> InputWidget;
+	TSharedPtr<SButton> InputWidget;
+	TSharedPtr<STextBlock> ButtonText;
 };
