Console Variables Editor:

-Make Input for showflags an int slider

#[fyi] Jason.Walter

diff --git a/Engine/Plugins/Experimental/ConsoleVariables/Source/ConsoleVariablesEditor/Private/Views/List/SConsoleVariablesEditorListValueInput.cpp b/Engine/Plugins/Experimental/ConsoleVariables/Source/ConsoleVariablesEditor/Private/Views/List/SConsoleVariablesEditorListValueInput.cpp
--- a/Engine/Plugins/Experimental/ConsoleVariables/Source/ConsoleVariablesEditor/Private/Views/List/SConsoleVariablesEditorListValueInput.cpp
+++ b/Engine/Plugins/Experimental/ConsoleVariables/Source/ConsoleVariablesEditor/Private/Views/List/SConsoleVariablesEditorListValueInput.cpp
@@ -47,6 +47,12 @@ TSharedRef<SConsoleVariablesEditorListValueInput> SConsoleVariablesEditorListVal
 			{
 				return SNew(SConsoleVariablesEditorListValueInput_String, InRow);
 			}
+
+			// Showflags are not considered to be any of these types, but they should be ints with a min/max of 0/2
+			if (PinnedInfo->Command.Contains("showflag", ESearchCase::IgnoreCase))
+			{
+				return SNew(SConsoleVariablesEditorListValueInput_Int, InRow, true);
+			}
 		}
 	}
 
@@ -124,7 +130,8 @@ float SConsoleVariablesEditorListValueInput_Float::GetInputValue() const
 }
 
 void SConsoleVariablesEditorListValueInput_Int::Construct(const FArguments& InArgs,
-                                                          const TWeakPtr<FConsoleVariablesEditorListRow> InRow)
+                                                          const TWeakPtr<FConsoleVariablesEditorListRow> InRow,
+                                                          const bool bIsShowFlag)
 {
 	check (InRow.IsValid());
 	
@@ -167,6 +174,12 @@ void SConsoleVariablesEditorListValueInput_Int::Construct(const FArguments& InAr
 		})
 		.IsEnabled(this, &SConsoleVariablesEditorListValueInput::IsRowChecked)
 	];
+
+	if (bIsShowFlag)
+	{
+		InputWidget->SetMinSliderValue(0);
+		InputWidget->SetMaxSliderValue(2);
+	}
 	
 	Item.Pin()->SetCachedValue(GetInputValueAsString());
 }
diff --git a/Engine/Plugins/Experimental/ConsoleVariables/Source/ConsoleVariablesEditor/Private/Views/List/SConsoleVariablesEditorListValueInput.h b/Engine/Plugins/Experimental/ConsoleVariables/Source/ConsoleVariablesEditor/Private/Views/List/SConsoleVariablesEditorListValueInput.h
--- a/Engine/Plugins/Experimental/ConsoleVariables/Source/ConsoleVariablesEditor/Private/Views/List/SConsoleVariablesEditorListValueInput.h
+++ b/Engine/Plugins/Experimental/ConsoleVariables/Source/ConsoleVariablesEditor/Private/Views/List/SConsoleVariablesEditorListValueInput.h
@@ -60,7 +60,8 @@ public:
 
 	SLATE_END_ARGS()
 
-	void Construct(const FArguments& InArgs, const TWeakPtr<FConsoleVariablesEditorListRow> InRow);
+	void Construct(
+		const FArguments& InArgs, const TWeakPtr<FConsoleVariablesEditorListRow> InRow, const bool bIsShowFlag = false);
 	
 	virtual void SetInputValue(const FString& InValueAsString) override;
 	virtual FString GetInputValueAsString() override;
