Issue 125273: Compile errors in ConsoleVariablesAsset.gen.cpp and ConsoleVariablesAsset.h

#[fyi] Jason.Walter

diff --git a/Engine/Plugins/Experimental/ConsoleVariables/Source/ConsoleVariablesEditor/Public/ConsoleVariablesAsset.h b/Engine/Plugins/Experimental/ConsoleVariables/Source/ConsoleVariablesEditor/Public/ConsoleVariablesAsset.h
--- a/Engine/Plugins/Experimental/ConsoleVariables/Source/ConsoleVariablesEditor/Public/ConsoleVariablesAsset.h
+++ b/Engine/Plugins/Experimental/ConsoleVariables/Source/ConsoleVariablesEditor/Public/ConsoleVariablesAsset.h
@@ -3,6 +3,7 @@
 #pragma once
 
 #include "CoreMinimal.h"
+#include "Widgets/Input/SCheckBox.h"
 
 #include "ConsoleVariablesAsset.generated.h"
 
