ConsoleVariables: Don't track cvars that have already been unregistered.

This avoids a crash when matching cheat variables like con.DebugLateCheat.

diff --git a/Engine/Plugins/Experimental/ConsoleVariables/Source/ConsoleVariablesEditor/Private/ConsoleVariablesEditorModule.cpp b/Engine/Plugins/Experimental/ConsoleVariables/Source/ConsoleVariablesEditor/Private/ConsoleVariablesEditorModule.cpp
--- a/Engine/Plugins/Experimental/ConsoleVariables/Source/ConsoleVariablesEditor/Private/ConsoleVariablesEditorModule.cpp
+++ b/Engine/Plugins/Experimental/ConsoleVariables/Source/ConsoleVariablesEditor/Private/ConsoleVariablesEditorModule.cpp
@@ -87,6 +87,11 @@ void FConsoleVariablesEditorModule::QueryAndBeginTrackingConsoleVariables()
 	IConsoleManager::Get().ForEachConsoleObjectThatStartsWith(FConsoleObjectVisitor::CreateLambda(
 		[this] (const TCHAR* Key, IConsoleObject* ConsoleObject)
 		{
+			if (!ConsoleObject || ConsoleObject->TestFlags(ECVF_Unregistered))
+			{
+				return;
+			}
+
 			if (IConsoleVariable* AsVariable = ConsoleObject->AsVariable())
 			{
 				const FDelegateHandle Handle =
