Console Variables Editor:

-Added support for console commands that are not variables
-Added an accessor to get the console object ptr that checks for staleness and periodically refreshes itself
-Added a delegate that fires when a console object is detected to have been unregistered
-When a cvar is detected to be unregistered the module will remove the cvar row and rebuild the list
-Fixed an issue in which similar float values were not being correctly compared
-Added Show Option that toggles bAddAllChangedConsoleVariablesToCurrentPreset
-Added a step to refresh the header checked state when refreshing the list
-Removed checkboxes from non-variable rows
-Updated comments in project settings to adhere to the standard

diff --git a/Engine/Plugins/Experimental/ConsoleVariables/Source/ConsoleVariablesEditor/Private/ConsoleVariablesAsset.cpp b/Engine/Plugins/Experimental/ConsoleVariables/Source/ConsoleVariablesEditor/Private/ConsoleVariablesAsset.cpp
--- a/Engine/Plugins/Experimental/ConsoleVariables/Source/ConsoleVariablesEditor/Private/ConsoleVariablesAsset.cpp
+++ b/Engine/Plugins/Experimental/ConsoleVariables/Source/ConsoleVariablesEditor/Private/ConsoleVariablesAsset.cpp
@@ -2,8 +2,6 @@
 
 #include "ConsoleVariablesAsset.h"
 
-#include "ConsoleVariablesEditorCommandInfo.h"
-
 #include "Algo/Find.h"
 
 void UConsoleVariablesAsset::SetVariableCollectionDescription(const FString& InVariableCollectionDescription)
@@ -32,7 +30,7 @@ bool UConsoleVariablesAsset::FindSavedDataByCommandString(const FString& InComma
 	return false;
 }
 
-void UConsoleVariablesAsset::AddOrSetConsoleVariableSavedData(const FConsoleVariablesEditorAssetSaveData& InData)
+void UConsoleVariablesAsset::AddOrSetConsoleObjectSavedData(const FConsoleVariablesEditorAssetSaveData& InData)
 {
 	RemoveConsoleVariable(InData.CommandName);
 	SavedCommands.Add(InData);
diff --git a/Engine/Plugins/Experimental/ConsoleVariables/Source/ConsoleVariablesEditor/Private/ConsoleVariablesEditorCommandInfo.h b/Engine/Plugins/Experimental/ConsoleVariables/Source/ConsoleVariablesEditor/Private/ConsoleVariablesEditorCommandInfo.h
--- a/Engine/Plugins/Experimental/ConsoleVariables/Source/ConsoleVariablesEditor/Private/ConsoleVariablesEditorCommandInfo.h
+++ b/Engine/Plugins/Experimental/ConsoleVariables/Source/ConsoleVariablesEditor/Private/ConsoleVariablesEditorCommandInfo.h
@@ -7,15 +7,23 @@
 #include "ConsoleVariablesEditorLog.h"
 #include "Editor.h"
 #include "Engine/GameEngine.h"
+#include "HAL/IConsoleManager.h"
 
-#include "ConsoleVariablesEditorCommandInfo.generated.h"
+DECLARE_MULTICAST_DELEGATE_OneParam(FOnDetectConsoleObjectUnregistered, FString)
 
 #define LOCTEXT_NAMESPACE "ConsoleVariablesEditor"
 
-USTRUCT()
 struct FConsoleVariablesEditorCommandInfo
 {
-	GENERATED_BODY()
+	enum class EConsoleObjectType
+	{
+		// A console command that has no associated console object but is parsed externally, e.g. 'stat unit'
+		NullObject,
+		// A console command with an associated console object, like 'r.SetNearClipPlane'
+		Command,
+		// A console variable such as 'r.ScreenPercentage'
+		Variable
+	};
 
 	struct FStaticConsoleVariableFlagInfo
 	{
@@ -23,31 +31,87 @@ struct FConsoleVariablesEditorCommandInfo
 		FText DisplayText;
 	};
 
-	FConsoleVariablesEditorCommandInfo()
+	FConsoleVariablesEditorCommandInfo(const FString& InCommand)
+	: Command(InCommand)
 	{
-		ConsoleVariablePtr = nullptr;
+		if (GetConsoleObjectPtr())
+		{
+			ObjectType = EConsoleObjectType::Command;
+			
+			if (const IConsoleVariable* AsVariable = GetConsoleVariablePtr())
+			{
+				ObjectType = EConsoleObjectType::Variable;
+				StartupValueAsString = AsVariable->GetString();
+				StartupSource = GetSource();
+			}
+		}
 	}
 
-	FConsoleVariablesEditorCommandInfo(
-		const FString& InCommand, IConsoleVariable* InVariablePtr, const FDelegateHandle& InOnVariableChangedCallbackHandle)
-	: Command(InCommand)
-	, ConsoleVariablePtr(InVariablePtr)
-	, OnVariableChangedCallbackHandle(InOnVariableChangedCallbackHandle)
+	~FConsoleVariablesEditorCommandInfo()
 	{
-		StartupValueAsString = ConsoleVariablePtr->GetString();
-		StartupSource = GetSource();
+		OnDetectConsoleObjectUnregistered.Remove(OnDetectConsoleObjectUnregisteredHandle);
+
+		if (IConsoleVariable* AsVariable = GetConsoleVariablePtr())
+		{
+			AsVariable->OnChangedDelegate().Remove(OnVariableChangedCallbackHandle);
+		}
 	}
 
 	FORCEINLINE bool operator==(const FConsoleVariablesEditorCommandInfo& Comparator) const
 	{
-		return ConsoleVariablePtr && Comparator.ConsoleVariablePtr &&
-			ConsoleVariablePtr == Comparator.ConsoleVariablePtr &&
-			Command.Equals(Comparator.Command);
+		return Command.Equals(Comparator.Command);
 	}
 
 	void ExecuteCommand(const FString& NewValueAsString) const
 	{
-		GEngine->Exec(GetCurrentWorld(), *FString::Printf(TEXT("%s %s"), *Command, *NewValueAsString));
+		GEngine->Exec(GetCurrentWorld(),
+			*FString::Printf(TEXT("%s %s"), *Command, *NewValueAsString).TrimStartAndEnd());
+	}
+
+	/** Get a reference to the cached console object. May return nullptr if unregistered. */
+	IConsoleObject* GetConsoleObjectPtr()
+	{
+		// If the console object ptr goes stale or is older than the specified threshold, try to refresh it
+		// May return nullptr if unregistered
+		if (!ConsoleObjectPtr ||
+			(FDateTime::UtcNow() - TimeOfLastConsoleObjectRefresh).GetTotalSeconds() > ConsoleObjectRefreshThreshold)
+		{
+			ConsoleObjectPtr = IConsoleManager::Get().FindConsoleObject(*Command);
+			TimeOfLastConsoleObjectRefresh = FDateTime::UtcNow();
+		}
+
+		// If the console object turns out to be unregistered, let interested parties know
+		if (ConsoleObjectPtr && ConsoleObjectPtr->TestFlags(ECVF_Unregistered))
+		{
+			OnDetectConsoleObjectUnregistered.Broadcast(Command);
+		}
+
+		return ConsoleObjectPtr;
+	}
+
+	/** Return the console object as a console variable object if applicable. May return nullptr if unregistered. */
+	IConsoleVariable* GetConsoleVariablePtr()
+	{
+		if (IConsoleObject* ObjectPtr = GetConsoleObjectPtr())
+		{
+			return ObjectPtr->AsVariable();
+		}
+		
+		return nullptr;
+	}
+
+	/**
+	 *Return the console object as a console command object if applicable.
+	 *Does not consider externally parsed console commands, as they have no associated objects.
+	 */
+	IConsoleCommand* GetConsoleCommandPtr()
+	{
+		if (IConsoleObject* ObjectPtr = GetConsoleObjectPtr())
+		{
+			return ObjectPtr->AsCommand();
+		}
+		
+		return nullptr;
 	}
 	
 	static UWorld* GetCurrentWorld()
@@ -64,31 +128,47 @@ struct FConsoleVariablesEditorCommandInfo
 		return CurrentWorld;
 	}
 
-	EConsoleVariableFlags GetSource() const
+	FString GetHelpText()
 	{
-		if (ConsoleVariablePtr)
+		if (const IConsoleVariable* AsVariable = GetConsoleVariablePtr())
 		{
-			return (EConsoleVariableFlags)((uint32)ConsoleVariablePtr->GetFlags() & ECVF_SetByMask);
+			return FString(AsVariable->GetHelp());
+		}
+
+		return "";
+	}
+
+	EConsoleVariableFlags GetSource()
+	{
+		if (const IConsoleObject* ConsoleObject = GetConsoleObjectPtr())
+		{
+			return (EConsoleVariableFlags)((uint32)ConsoleObject->GetFlags() & ECVF_SetByMask);
 		}
 
 		return ECVF_Default;
 	}
 
-	void ClearSourceFlags() const
+	void ClearSourceFlags()
 	{
-		for (const FStaticConsoleVariableFlagInfo& StaticConsoleVariableFlagInfo : SupportedFlags)
+		if (IConsoleObject* ConsoleObject = GetConsoleObjectPtr())
 		{
-			ConsoleVariablePtr->ClearFlags(StaticConsoleVariableFlagInfo.Flag);
+			for (const FStaticConsoleVariableFlagInfo& StaticConsoleVariableFlagInfo : SupportedFlags)
+			{
+				ConsoleObject->ClearFlags(StaticConsoleVariableFlagInfo.Flag);
+			}
 		}
 	}
 
-	void SetSourceFlag(const EConsoleVariableFlags InSource) const
+	void SetSourceFlag(const EConsoleVariableFlags InSource)
 	{
-		ClearSourceFlags();
-		ConsoleVariablePtr->SetFlags((EConsoleVariableFlags)InSource);
+		if (IConsoleObject* ConsoleObject = GetConsoleObjectPtr())
+		{
+			ClearSourceFlags();
+			ConsoleObject->SetFlags((EConsoleVariableFlags)InSource);
+		}
 	}
 
-	FText GetSourceAsText() const
+	FText GetSourceAsText()
 	{
 		return ConvertConsoleVariableSetByFlagToText(GetSource());
 	}
@@ -97,7 +177,8 @@ struct FConsoleVariablesEditorCommandInfo
 	{
 		FText ReturnValue = LOCTEXT("UnknownSource", "<UNKNOWN>"); 
 
-		if (const FStaticConsoleVariableFlagInfo* Match = Algo::FindByPredicate(SupportedFlags,
+		if (const FStaticConsoleVariableFlagInfo* Match = Algo::FindByPredicate(
+			SupportedFlags,
 				[InFlag](const FStaticConsoleVariableFlagInfo& Comparator)
 				{
 					return Comparator.Flag == InFlag;
@@ -109,29 +190,51 @@ struct FConsoleVariablesEditorCommandInfo
 		return ReturnValue;
 	}
 
-	bool IsCurrentValueDifferentFromInputValue(const FString& InValueToCompare) const
+	bool IsCurrentValueDifferentFromInputValue(const FString& InValueToCompare)
 	{
-		if (ConsoleVariablePtr)
+		if (const IConsoleVariable* AsVariable = GetConsoleVariablePtr())
 		{
-			return !ConsoleVariablePtr->GetString().Equals(InValueToCompare);
+			// Floats sometimes return true erroneously because they can be stringified as e.g '1' or '1.0' by different functions.
+			if (AsVariable->IsVariableFloat())
+			{
+				const float A = AsVariable->GetFloat();
+				const float B = FCString::Atof(*InValueToCompare);
+
+				return !FMath::IsNearlyEqual(A, B);
+			}
+			
+			return !AsVariable->GetString().Equals(InValueToCompare);
 		}
 
 		return false;
 	}
 
-	/** The actual string command to execute */
+	/** The actual string key or name */
 	UPROPERTY()
 	FString Command;
 
-	IConsoleVariable* ConsoleVariablePtr;
+	EConsoleObjectType ObjectType = EConsoleObjectType::NullObject;
+
+	/** This object is periodically refreshed to mitigate the occurrence of stale pointers. */
+	IConsoleObject* ConsoleObjectPtr;
+	FDateTime TimeOfLastConsoleObjectRefresh;
 
-	/** The value of this command when the module started in this session after it may have been set by an ini file. */
+	double ConsoleObjectRefreshThreshold = 1.0;
+	
+	/** The value of this variable (if Variable object type) when the module started in this session after it may have been set by an ini file. */
 	FString StartupValueAsString;
-	/** The source of this variable last setting as recorded when the plugin was loaded. */
-	EConsoleVariableFlags StartupSource;
 	
+	/** The source of this variable's (if Variable object type) last setting as recorded when the plugin was loaded. */
+	EConsoleVariableFlags StartupSource = ECVF_Default;
+
+	/** When variables change, this callback is executed. */
 	FDelegateHandle OnVariableChangedCallbackHandle;
 
+	/** When commands are unregistered change, this callback is broadcasted. */
+	FOnDetectConsoleObjectUnregistered OnDetectConsoleObjectUnregistered;
+	FDelegateHandle OnDetectConsoleObjectUnregisteredHandle;
+
+	/** A mapping of SetBy console variable flags to information like the associated display text. */
 	static const inline TArray<FStaticConsoleVariableFlagInfo> SupportedFlags =
 	{
 		{ EConsoleVariableFlags::ECVF_SetByConstructor, LOCTEXT("SetByConstructor", "Constructor") },
diff --git a/Engine/Plugins/Experimental/ConsoleVariables/Source/ConsoleVariablesEditor/Private/ConsoleVariablesEditorModule.cpp b/Engine/Plugins/Experimental/ConsoleVariables/Source/ConsoleVariablesEditor/Private/ConsoleVariablesEditorModule.cpp
--- a/Engine/Plugins/Experimental/ConsoleVariables/Source/ConsoleVariablesEditor/Private/ConsoleVariablesEditorModule.cpp
+++ b/Engine/Plugins/Experimental/ConsoleVariables/Source/ConsoleVariablesEditor/Private/ConsoleVariablesEditorModule.cpp
@@ -47,22 +47,13 @@ void FConsoleVariablesEditorModule::ShutdownModule()
 	
 	MainPanel.Reset();
 
-	ConsoleVariablesMasterReference.Empty();
+	ConsoleObjectsMasterReference.Empty();
 
 	// Unregister project settings
 	ISettingsModule& SettingsModule = FModuleManager::LoadModuleChecked<ISettingsModule>("Settings");
 	{
 		SettingsModule.UnregisterSettings("Project", "Plugins", "Console Variables Editor");
 	}
-
-	// Remove all OnChanged delegates
-	for (TSharedPtr<FConsoleVariablesEditorCommandInfo> CommandInfo : ConsoleVariablesMasterReference)
-	{
-		if (CommandInfo.IsValid() && CommandInfo->ConsoleVariablePtr)
-		{
-			CommandInfo->ConsoleVariablePtr->OnChangedDelegate().Remove(CommandInfo->OnVariableChangedCallbackHandle);
-		}
-	}
 }
 
 void FConsoleVariablesEditorModule::OpenConsoleVariablesDialogWithAssetSelected(const FAssetData& InAssetData)
@@ -80,9 +71,9 @@ void FConsoleVariablesEditorModule::OpenConsoleVariablesDialogWithAssetSelected(
 
 void FConsoleVariablesEditorModule::QueryAndBeginTrackingConsoleVariables()
 {
-	const int32 VariableCount = ConsoleVariablesMasterReference.Num();
+	const int32 VariableCount = ConsoleObjectsMasterReference.Num();
 	
-	ConsoleVariablesMasterReference.Empty(VariableCount);
+	ConsoleObjectsMasterReference.Empty(VariableCount);
 	
 	IConsoleManager::Get().ForEachConsoleObjectThatStartsWith(FConsoleObjectVisitor::CreateLambda(
 		[this] (const TCHAR* Key, IConsoleObject* ConsoleObject)
@@ -92,15 +83,20 @@ void FConsoleVariablesEditorModule::QueryAndBeginTrackingConsoleVariables()
 				return;
 			}
 
+			const TSharedRef<FConsoleVariablesEditorCommandInfo> Info =
+				MakeShared<FConsoleVariablesEditorCommandInfo>(Key);
+			
+			Info->StartupSource = Info->GetSource();
+			Info->OnDetectConsoleObjectUnregisteredHandle = Info->OnDetectConsoleObjectUnregistered.AddRaw(
+				this, &FConsoleVariablesEditorModule::OnDetectConsoleObjectUnregistered);
+
 			if (IConsoleVariable* AsVariable = ConsoleObject->AsVariable())
 			{
-				const FDelegateHandle Handle =
-					AsVariable->OnChangedDelegate().AddRaw(this, &FConsoleVariablesEditorModule::OnConsoleVariableChanged);
-				const TSharedRef<FConsoleVariablesEditorCommandInfo> Info =
-					MakeShared<FConsoleVariablesEditorCommandInfo>(Key, AsVariable, Handle);
-				Info->StartupSource = Info->GetSource();
-				ConsoleVariablesMasterReference.Add(Info);
+				Info->OnVariableChangedCallbackHandle = AsVariable->OnChangedDelegate().AddRaw(
+					this, &FConsoleVariablesEditorModule::OnConsoleVariableChanged);
 			}
+			
+			AddConsoleObjectCommandInfoToMasterReference(Info);
 		}),
 		TEXT(""));
 }
@@ -108,7 +104,7 @@ void FConsoleVariablesEditorModule::QueryAndBeginTrackingConsoleVariables()
 TWeakPtr<FConsoleVariablesEditorCommandInfo> FConsoleVariablesEditorModule::FindCommandInfoByName(const FString& NameToSearch, ESearchCase::Type InSearchCase)
 {
 	TSharedPtr<FConsoleVariablesEditorCommandInfo>* Match = Algo::FindByPredicate(
-		ConsoleVariablesMasterReference,
+		ConsoleObjectsMasterReference,
 		[&NameToSearch, InSearchCase](const TSharedPtr<FConsoleVariablesEditorCommandInfo> Comparator)
 		{
 			return Comparator->Command.Equals(NameToSearch, InSearchCase);
@@ -117,13 +113,14 @@ TWeakPtr<FConsoleVariablesEditorCommandInfo> FConsoleVariablesEditorModule::Find
 	return Match ? *Match : nullptr;
 }
 
-TWeakPtr<FConsoleVariablesEditorCommandInfo> FConsoleVariablesEditorModule::FindCommandInfoByConsoleVariableReference(IConsoleVariable* InVariableReference)
+TWeakPtr<FConsoleVariablesEditorCommandInfo> FConsoleVariablesEditorModule::FindCommandInfoByConsoleObjectReference(
+	IConsoleObject* InConsoleObjectReference)
 {
 	TSharedPtr<FConsoleVariablesEditorCommandInfo>* Match = Algo::FindByPredicate(
-	ConsoleVariablesMasterReference,
-	[InVariableReference](const TSharedPtr<FConsoleVariablesEditorCommandInfo> Comparator)
+	ConsoleObjectsMasterReference,
+	[InConsoleObjectReference](const TSharedPtr<FConsoleVariablesEditorCommandInfo> Comparator)
 	{
-		return Comparator->ConsoleVariablePtr == InVariableReference;
+		return Comparator->GetConsoleObjectPtr() == InConsoleObjectReference;
 	});
 
 	return Match ? *Match : nullptr;
@@ -186,7 +183,7 @@ void FConsoleVariablesEditorModule::OnConsoleVariableChanged(IConsoleVariable* C
 	check(EditingAsset);
 
 	if (const TWeakPtr<FConsoleVariablesEditorCommandInfo> CommandInfo =
-		FindCommandInfoByConsoleVariableReference(ChangedVariable); CommandInfo.IsValid())
+		FindCommandInfoByConsoleObjectReference(ChangedVariable); CommandInfo.IsValid())
 	{
 		const TSharedPtr<FConsoleVariablesEditorCommandInfo>& PinnedCommand = CommandInfo.Pin();
 		const FString& Key = PinnedCommand->Command;
@@ -203,7 +200,11 @@ void FConsoleVariablesEditorModule::OnConsoleVariableChanged(IConsoleVariable* C
 			{
 				if (MainPanel.IsValid())
 				{
-					MainPanel->AddConsoleVariable(Key, ChangedVariable->GetString());
+					MainPanel->AddConsoleObjectToPreset(
+						Key,
+						ChangedVariable->GetString(),
+						true
+					);
 				}
 
 				SendMultiUserConsoleVariableChange(Key, ChangedVariable->GetString());
@@ -221,6 +222,24 @@ void FConsoleVariablesEditorModule::OnConsoleVariableChanged(IConsoleVariable* C
 	}
 }
 
+void FConsoleVariablesEditorModule::OnDetectConsoleObjectUnregistered(FString CommandName)
+{
+	check(EditingAsset);
+
+	EditingAsset->RemoveConsoleVariable(CommandName);
+
+	if (MainPanel.IsValid())
+	{
+		MainPanel->RefreshList();
+	}
+
+	if (const TWeakPtr<FConsoleVariablesEditorCommandInfo> CommandInfo =
+		FindCommandInfoByName(CommandName); CommandInfo.IsValid())
+	{
+		ConsoleObjectsMasterReference.Remove(CommandInfo.Pin());
+	}
+}
+
 void FConsoleVariablesEditorModule::OnRemoteCvarChanged(const FString InName, const FString InValue)
 {
 	UE_LOG(LogConsoleVariablesEditor, Display, TEXT("Remote set console variable %s = %s"), *InName, *InValue);
diff --git a/Engine/Plugins/Experimental/ConsoleVariables/Source/ConsoleVariablesEditor/Private/ConsoleVariablesEditorModule.h b/Engine/Plugins/Experimental/ConsoleVariables/Source/ConsoleVariablesEditor/Private/ConsoleVariablesEditorModule.h
--- a/Engine/Plugins/Experimental/ConsoleVariables/Source/ConsoleVariablesEditor/Private/ConsoleVariablesEditorModule.h
+++ b/Engine/Plugins/Experimental/ConsoleVariables/Source/ConsoleVariablesEditor/Private/ConsoleVariablesEditorModule.h
@@ -33,11 +33,22 @@ public:
 	/** Find all console variables and cache their startup values */
 	void QueryAndBeginTrackingConsoleVariables();
 
-	/** Find a tracked console variable by the command string with optional case sensitivity. */
-	TWeakPtr<FConsoleVariablesEditorCommandInfo> FindCommandInfoByName(const FString& NameToSearch, ESearchCase::Type InSearchCase = ESearchCase::IgnoreCase);
+	void AddConsoleObjectCommandInfoToMasterReference(TSharedRef<FConsoleVariablesEditorCommandInfo> InCommandInfo)
+	{
+		ConsoleObjectsMasterReference.Add(InCommandInfo);
+	}
 
-	/** Find a tracked console variable by its variable reference. */
-	TWeakPtr<FConsoleVariablesEditorCommandInfo> FindCommandInfoByConsoleVariableReference(IConsoleVariable* InVariableReference);
+	/** Find a tracked console variable by the command string with optional case sensitivity. */
+	TWeakPtr<FConsoleVariablesEditorCommandInfo> FindCommandInfoByName(
+		const FString& NameToSearch, ESearchCase::Type InSearchCase = ESearchCase::IgnoreCase);
+
+	/**
+	 *Find a tracked console variable by its console object reference.
+	 *Note that some commands do not have an associated console object (such as 'stat unit') and will not be found with this method.
+	 *It's normally safer to use FindCommandInfoByName() instead.
+	 */
+	TWeakPtr<FConsoleVariablesEditorCommandInfo> FindCommandInfoByConsoleObjectReference(
+		IConsoleObject* InConsoleObjectReference);
 	
 	[[nodiscard]] TObjectPtr<UConsoleVariablesAsset> GetEditingAsset() const;
 	void SetEditingAsset(const TObjectPtr<UConsoleVariablesAsset> InEditingAsset);
@@ -53,6 +64,8 @@ private:
 	void RegisterProjectSettings();
 
 	void OnConsoleVariableChanged(IConsoleVariable* ChangedVariable);
+	/** In the event a console object is unregistered, this failsafe callback will clean up the associated list item and command info object. */
+	void OnDetectConsoleObjectUnregistered(FString CommandName);
 
 	TObjectPtr<UConsoleVariablesAsset> AllocateTransientPreset();
 
@@ -69,5 +82,5 @@ private:
 	TObjectPtr<UConsoleVariablesAsset> EditingAsset = nullptr;
 
 	/** All tracked variables and their default, startup, and current values */
-	TArray<TSharedPtr<FConsoleVariablesEditorCommandInfo>> ConsoleVariablesMasterReference;
+	TArray<TSharedPtr<FConsoleVariablesEditorCommandInfo>> ConsoleObjectsMasterReference;
 };
diff --git a/Engine/Plugins/Experimental/ConsoleVariables/Source/ConsoleVariablesEditor/Private/Views/List/ConsoleVariablesEditorListRow.cpp b/Engine/Plugins/Experimental/ConsoleVariables/Source/ConsoleVariablesEditor/Private/Views/List/ConsoleVariablesEditorListRow.cpp
--- a/Engine/Plugins/Experimental/ConsoleVariables/Source/ConsoleVariablesEditor/Private/Views/List/ConsoleVariablesEditorListRow.cpp
+++ b/Engine/Plugins/Experimental/ConsoleVariables/Source/ConsoleVariablesEditor/Private/Views/List/ConsoleVariablesEditorListRow.cpp
@@ -122,10 +122,15 @@ bool FConsoleVariablesEditorListRow::MatchSearchTokensToSearchTerms(const TArray
 
 	if (!bMatchFound)
 	{
-		TSharedPtr<FConsoleVariablesEditorCommandInfo> PinnedInfo = CommandInfo.Pin();
+		const TSharedPtr<FConsoleVariablesEditorCommandInfo> PinnedInfo = CommandInfo.Pin();
 		
-		const FString SearchTerms = PinnedInfo->Command + PinnedInfo->GetSourceAsText().ToString() +
-			(PinnedInfo->ConsoleVariablePtr ? PinnedInfo->ConsoleVariablePtr->GetString() + PinnedInfo->ConsoleVariablePtr->GetHelp() : "");
+		FString SearchTerms = PinnedInfo->Command + PinnedInfo->GetSourceAsText().ToString();
+
+		if (const IConsoleVariable* AsVariable = PinnedInfo->GetConsoleVariablePtr())
+		{
+			// Add ValueAsString and HelpText to SearchTerms
+			SearchTerms += AsVariable->GetString() + PinnedInfo->GetHelpText();
+		}
 
 		bMatchFound = Algo::AnyOf(InTokens,
 			[&SearchTerms](const FString& Token)
diff --git a/Engine/Plugins/Experimental/ConsoleVariables/Source/ConsoleVariablesEditor/Private/Views/List/SConsoleVariablesEditorList.cpp b/Engine/Plugins/Experimental/ConsoleVariables/Source/ConsoleVariablesEditor/Private/Views/List/SConsoleVariablesEditorList.cpp
--- a/Engine/Plugins/Experimental/ConsoleVariables/Source/ConsoleVariablesEditor/Private/Views/List/SConsoleVariablesEditorList.cpp
+++ b/Engine/Plugins/Experimental/ConsoleVariables/Source/ConsoleVariablesEditor/Private/Views/List/SConsoleVariablesEditorList.cpp
@@ -5,6 +5,7 @@
 #include "ConsoleVariablesAsset.h"
 #include "ConsoleVariablesEditorListFilters/ConsoleVariablesEditorListFilter_SourceText.h"
 #include "ConsoleVariablesEditorModule.h"
+#include "ConsoleVariablesEditorProjectSettings.h"
 #include "SConsoleVariablesEditorListRow.h"
 
 #include "Algo/Find.h"
@@ -163,7 +164,7 @@ TSharedRef<SWidget> SConsoleVariablesEditorList::BuildShowOptionsMenu()
 {
 	FMenuBuilder ShowOptionsMenuBuilder = FMenuBuilder(true, nullptr);
 
-	ShowOptionsMenuBuilder.BeginSection("AssetThumbnails", LOCTEXT("ShowOptionsShowSectionHeading", "Show"));
+	ShowOptionsMenuBuilder.BeginSection("", LOCTEXT("ShowOptions_ShowSectionHeading", "Show"));
 	{
 		// Add mode filters
 		auto AddFiltersLambda = [this, &ShowOptionsMenuBuilder](const TSharedRef<IConsoleVariablesEditorListFilter>& InFilter)
@@ -196,7 +197,7 @@ TSharedRef<SWidget> SConsoleVariablesEditorList::BuildShowOptionsMenu()
 	}
 	ShowOptionsMenuBuilder.EndSection();
 	
-	ShowOptionsMenuBuilder.BeginSection("AssetThumbnails", LOCTEXT("SortHeading", "Sort"));
+	ShowOptionsMenuBuilder.BeginSection("", LOCTEXT("ShowOptions_SortSectionHeading", "Sort"));
 	{
 		// Add commands
 
@@ -220,6 +221,41 @@ TSharedRef<SWidget> SConsoleVariablesEditorList::BuildShowOptionsMenu()
 		);
 	}
 	ShowOptionsMenuBuilder.EndSection();
+	
+	ShowOptionsMenuBuilder.BeginSection("", LOCTEXT("ShowOptions_OptionsSectionHeading", "Options"));
+	{
+		ShowOptionsMenuBuilder.AddMenuEntry(
+			LOCTEXT("TrackAllVariableChanges", "Track All Variable Changes"),
+			LOCTEXT("ConsoleVariablesEditorList_TrackAllVariableChanges_Tooltip",
+				"When variables are changed outside the Console Variables Editor, this option will add the variables to the current preset. Does not apply to console commands like 'r.SetNearClipPlane' or 'stat fps'."),
+			FSlateIcon(),
+			FUIAction(
+				FExecuteAction::CreateLambda([]()
+				{
+					if (UConsoleVariablesEditorProjectSettings* ProjectSettingsPtr =
+						GetMutableDefault<UConsoleVariablesEditorProjectSettings>())
+					{
+						ProjectSettingsPtr->bAddAllChangedConsoleVariablesToCurrentPreset =
+							!ProjectSettingsPtr->bAddAllChangedConsoleVariablesToCurrentPreset;
+					}
+				}),
+			   FCanExecuteAction(),
+			   FIsActionChecked::CreateLambda( []()
+				{
+					if (const UConsoleVariablesEditorProjectSettings* ProjectSettingsPtr =
+						GetMutableDefault<UConsoleVariablesEditorProjectSettings>())
+					{
+						return ProjectSettingsPtr->bAddAllChangedConsoleVariablesToCurrentPreset;
+					}
+
+			   		return false;
+				})
+			),
+			NAME_None,
+			EUserInterfaceActionType::ToggleButton
+		);
+	}
+	ShowOptionsMenuBuilder.EndSection();
 
 	return ShowOptionsMenuBuilder.MakeWidget();
 }
@@ -277,6 +313,9 @@ void SConsoleVariablesEditorList::RefreshList()
 		// Show/Hide rows based on SetBy changes and filter settings
 		EvaluateIfRowsPassFilters();
 
+		// Refresh the header's check state
+		OnListItemCheckBoxStateChange(ECheckBoxState::Undetermined);
+
 		TreeViewPtr->RequestTreeRefresh();
 	}
 }
@@ -343,9 +382,19 @@ void SConsoleVariablesEditorList::GenerateTreeView()
 
 	for (const FConsoleVariablesEditorAssetSaveData& SavedCommand : EditableAsset->GetSavedCommands())
 	{
-		// Get corresponding CommandInfo for tracking
-		if (TWeakPtr<FConsoleVariablesEditorCommandInfo> CommandInfo =
-			ConsoleVariablesEditorModule.FindCommandInfoByName(SavedCommand.CommandName); CommandInfo.IsValid())
+		// Get corresponding CommandInfo for tracking or make one if the command is non-value
+		TSharedPtr<FConsoleVariablesEditorCommandInfo> CommandInfo =
+			ConsoleVariablesEditorModule.FindCommandInfoByName(SavedCommand.CommandName).Pin();
+
+		if (!CommandInfo.IsValid())
+		{
+			CommandInfo = MakeShared<FConsoleVariablesEditorCommandInfo>(SavedCommand.CommandName);
+
+			FConsoleVariablesEditorModule::Get().AddConsoleObjectCommandInfoToMasterReference(
+				CommandInfo.ToSharedRef());
+		}
+		
+		if (CommandInfo.IsValid())
 		{
 			const ECheckBoxState NewCheckedState =
 				SavedCommand.CheckedState == ECheckBoxState::Unchecked ?
@@ -353,15 +402,16 @@ void SConsoleVariablesEditorList::GenerateTreeView()
 			
 			// If the row is checked and the saved value differs from the current value,
 			// execute the command with the saved value 
-			if ( NewCheckedState == ECheckBoxState::Checked &&
-				CommandInfo.Pin()->IsCurrentValueDifferentFromInputValue(SavedCommand.CommandValueAsString))
+			if (CommandInfo->ObjectType == FConsoleVariablesEditorCommandInfo::EConsoleObjectType::Command ||
+				(NewCheckedState == ECheckBoxState::Checked &&
+				CommandInfo->IsCurrentValueDifferentFromInputValue(SavedCommand.CommandValueAsString)))
 			{
-				CommandInfo.Pin()->ExecuteCommand(SavedCommand.CommandValueAsString);
+				CommandInfo->ExecuteCommand(SavedCommand.CommandValueAsString);
 			}
 		
 			FConsoleVariablesEditorListRowPtr NewRow = 
 				MakeShared<FConsoleVariablesEditorListRow>(
-						CommandInfo.Pin(), SavedCommand.CommandValueAsString,
+						CommandInfo, SavedCommand.CommandValueAsString,
 						FConsoleVariablesEditorListRow::SingleCommand, 
 						NewCheckedState, SharedThis(this), TreeViewRootObjects.Num(), nullptr);
 		
@@ -407,7 +457,11 @@ TSharedPtr<SHeaderRow> SConsoleVariablesEditorList::GenerateHeaderRow()
 					
 					for (const FConsoleVariablesEditorListRowPtr& Object : TreeViewRootObjects)
 					{
-						Object->SetWidgetCheckedState(NewState);
+						if (Object->GetCommandInfo().Pin()->ObjectType ==
+							FConsoleVariablesEditorCommandInfo::EConsoleObjectType::Variable)
+						{
+							Object->SetWidgetCheckedState(NewState);
+						}
 					}
 				})
 			]
@@ -538,7 +592,9 @@ bool SConsoleVariablesEditorList::DoesListHaveCheckedMembers() const
 {
 	for (const TSharedPtr<FConsoleVariablesEditorListRow>& Row : TreeViewRootObjects)
 	{
-		if (Row->GetWidgetCheckedState() == ECheckBoxState::Checked)
+		if (Row->GetCommandInfo().Pin()->ObjectType ==
+			FConsoleVariablesEditorCommandInfo::EConsoleObjectType::Variable &&
+			Row->GetWidgetCheckedState() == ECheckBoxState::Checked)
 		{
 			return true;
 		}
@@ -551,7 +607,9 @@ bool SConsoleVariablesEditorList::DoesListHaveUncheckedMembers() const
 {
 	for (const TSharedPtr<FConsoleVariablesEditorListRow>& Row : TreeViewRootObjects)
 	{
-		if (Row->GetWidgetCheckedState() == ECheckBoxState::Unchecked)
+		if (Row->GetCommandInfo().Pin()->ObjectType ==
+			FConsoleVariablesEditorCommandInfo::EConsoleObjectType::Variable &&
+			Row->GetWidgetCheckedState() == ECheckBoxState::Unchecked)
 		{
 			return true;
 		}
diff --git a/Engine/Plugins/Experimental/ConsoleVariables/Source/ConsoleVariablesEditor/Private/Views/List/SConsoleVariablesEditorListRow.cpp b/Engine/Plugins/Experimental/ConsoleVariables/Source/ConsoleVariablesEditor/Private/Views/List/SConsoleVariablesEditorListRow.cpp
--- a/Engine/Plugins/Experimental/ConsoleVariables/Source/ConsoleVariablesEditor/Private/Views/List/SConsoleVariablesEditorListRow.cpp
+++ b/Engine/Plugins/Experimental/ConsoleVariables/Source/ConsoleVariablesEditor/Private/Views/List/SConsoleVariablesEditorListRow.cpp
@@ -68,11 +68,6 @@ void SConsoleVariablesEditorListRow::Construct(
 	
 	check (PinnedCommand.IsValid());
 
-	if (!PinnedCommand->ConsoleVariablePtr)
-	{
-		return;
-	}
-
 	// Set up flash animation
 	FlashAnimation = FCurveSequence(0.f, FlashAnimationDuration, ECurveEaseFunction::QuadInOut);
 
@@ -355,6 +350,9 @@ TSharedRef<SWidget> SConsoleVariablesEditorListRow::GenerateCells(const FName& I
 					SNew(SCheckBox)
 					.IsChecked_Raw(this, &SConsoleVariablesEditorListRow::GetCheckboxState)
 					.OnCheckStateChanged_Raw(this, &SConsoleVariablesEditorListRow::OnCheckboxStateChange)
+					.Visibility(Item.Pin()->GetCommandInfo().Pin()->ObjectType ==
+						FConsoleVariablesEditorCommandInfo::EConsoleObjectType::Variable ?
+						EVisibility::Visible : EVisibility::Collapsed)
 				];
 	}
 	if (InColumnName.IsEqual(SConsoleVariablesEditorList::VariableNameColumnName))
@@ -363,7 +361,7 @@ TSharedRef<SWidget> SConsoleVariablesEditorListRow::GenerateCells(const FName& I
 		{
 			HoverToolTip = SConsoleVariablesEditorTooltipWidget::MakeTooltip(
 							PinnedItem->GetCommandInfo().Pin()->Command,
-							PinnedItem->GetCommandInfo().Pin()->ConsoleVariablePtr->GetHelp());
+							PinnedItem->GetCommandInfo().Pin()->GetHelpText());
 		}
 		return  SNew(STextBlock)
 				.Visibility(EVisibility::Visible)
@@ -405,6 +403,13 @@ TSharedRef<SWidget> SConsoleVariablesEditorListRow::GenerateCells(const FName& I
 
 ECheckBoxState SConsoleVariablesEditorListRow::GetCheckboxState() const
 {
+	//Non-variable types always checked
+	if (Item.Pin()->GetCommandInfo().Pin()->ObjectType !=
+		FConsoleVariablesEditorCommandInfo::EConsoleObjectType::Variable)
+	{
+		return ECheckBoxState::Checked;
+	}
+	
 	if (ensure(Item.IsValid()))
 	{
 		return Item.Pin()->GetWidgetCheckedState();
@@ -446,7 +451,7 @@ void SConsoleVariablesEditorListRow::OnCheckboxStateChange(const ECheckBoxState
 
 TSharedRef<SWidget> SConsoleVariablesEditorListRow::GenerateValueCellWidget(const TSharedPtr<FConsoleVariablesEditorListRow> PinnedItem)
 {
-	if (PinnedItem->GetCommandInfo().IsValid() && PinnedItem->GetCommandInfo().Pin()->ConsoleVariablePtr)
+	if (PinnedItem->GetCommandInfo().IsValid())
 	{
 		ValueChildInputWidget = SConsoleVariablesEditorListValueInput::GetInputWidget(Item);
 		const TSharedRef<SHorizontalBox> FinalValueWidget =
diff --git a/Engine/Plugins/Experimental/ConsoleVariables/Source/ConsoleVariablesEditor/Private/Views/List/SConsoleVariablesEditorListValueInput.cpp b/Engine/Plugins/Experimental/ConsoleVariables/Source/ConsoleVariablesEditor/Private/Views/List/SConsoleVariablesEditorListValueInput.cpp
--- a/Engine/Plugins/Experimental/ConsoleVariables/Source/ConsoleVariablesEditor/Private/Views/List/SConsoleVariablesEditorListValueInput.cpp
+++ b/Engine/Plugins/Experimental/ConsoleVariables/Source/ConsoleVariablesEditor/Private/Views/List/SConsoleVariablesEditorListValueInput.cpp
@@ -26,7 +26,7 @@ TSharedRef<SConsoleVariablesEditorListValueInput> SConsoleVariablesEditorListVal
 	{
 		const TSharedPtr<FConsoleVariablesEditorCommandInfo> PinnedInfo = PinnedItem->GetCommandInfo().Pin();
 
-		if (TObjectPtr<IConsoleVariable> Variable = PinnedInfo->ConsoleVariablePtr)
+		if (const TObjectPtr<IConsoleVariable> Variable = PinnedInfo->GetConsoleVariablePtr())
 		{
 			if (Variable->IsVariableFloat())
 			{
@@ -54,6 +54,9 @@ TSharedRef<SConsoleVariablesEditorListValueInput> SConsoleVariablesEditorListVal
 				return SNew(SConsoleVariablesEditorListValueInput_Int, InRow, true);
 			}
 		}
+
+		// For Commands
+		return SNew(SConsoleVariablesEditorListValueInput_Command, InRow, PinnedItem->GetPresetValue());
 	}
 
 	// fallback
@@ -86,7 +89,10 @@ void SConsoleVariablesEditorListValueInput_Float::Construct(const FArguments& In
 					(ProjectSettingsPtr &&
 						ProjectSettingsPtr->UncheckedRowDisplayType == EConsoleVariablesEditorRowDisplayType::ShowCurrentValue))
 			{
-				return FCString::Atof(*Item.Pin()->GetCommandInfo().Pin()->ConsoleVariablePtr->GetString());
+				if (const IConsoleVariable* AsVariable = Item.Pin()->GetCommandInfo().Pin()->GetConsoleVariablePtr())
+				{
+					return FCString::Atof(*AsVariable->GetString());
+				}
 			}
 
 			return FCString::Atof(*Item.Pin()->GetCachedValue());
@@ -150,7 +156,10 @@ void SConsoleVariablesEditorListValueInput_Int::Construct(const FArguments& InAr
 					(ProjectSettingsPtr &&
 						ProjectSettingsPtr->UncheckedRowDisplayType == EConsoleVariablesEditorRowDisplayType::ShowCurrentValue))
 			{
-				return FCString::Atoi(*Item.Pin()->GetCommandInfo().Pin()->ConsoleVariablePtr->GetString());
+				if (const IConsoleVariable* AsVariable = Item.Pin()->GetCommandInfo().Pin()->GetConsoleVariablePtr())
+				{
+					return FCString::Atoi(*AsVariable->GetString());
+				}
 			}
 
 			return FCString::Atoi(*Item.Pin()->GetCachedValue());
@@ -235,7 +244,10 @@ void SConsoleVariablesEditorListValueInput_String::Construct(const FArguments& I
 					(ProjectSettingsPtr &&
 						ProjectSettingsPtr->UncheckedRowDisplayType == EConsoleVariablesEditorRowDisplayType::ShowCurrentValue))
 			{
-				return FText::FromString(*Item.Pin()->GetCommandInfo().Pin()->ConsoleVariablePtr->GetString());
+				if (const IConsoleVariable* AsVariable = Item.Pin()->GetCommandInfo().Pin()->GetConsoleVariablePtr())
+				{
+					return FText::FromString(*AsVariable->GetString());
+				}
 			}
 
 			return FText::FromString(Item.Pin()->GetCachedValue());
@@ -310,7 +322,10 @@ void SConsoleVariablesEditorListValueInput_Bool::Construct(const FArguments& InA
 						(ProjectSettingsPtr &&
 							ProjectSettingsPtr->UncheckedRowDisplayType == EConsoleVariablesEditorRowDisplayType::ShowCurrentValue))
 				{
-					return FText::FromString(Item.Pin()->GetCommandInfo().Pin()->ConsoleVariablePtr->GetString());
+					if (const IConsoleVariable* AsVariable = Item.Pin()->GetCommandInfo().Pin()->GetConsoleVariablePtr())
+					{
+						return FText::FromString(AsVariable->GetString());
+					}
 				}
 
 				return FText::FromString(Item.Pin()->GetCachedValue());
@@ -351,4 +366,89 @@ bool SConsoleVariablesEditorListValueInput_Bool::GetInputValue()
 	return StringToBool(GetInputValueAsString());
 }
 
+void SConsoleVariablesEditorListValueInput_Command::Construct(const FArguments& InArgs,
+	const TWeakPtr<FConsoleVariablesEditorListRow> InRow, const FString& InSavedText)
+{
+	check (InRow.IsValid());
+	
+	Item = InRow;
+
+	ProjectSettingsPtr = GetMutableDefault<UConsoleVariablesEditorProjectSettings>();
+	
+	ChildSlot
+	[
+		SNew(SHorizontalBox)
+
+		+SHorizontalBox::Slot()
+		.VAlign(VAlign_Center)
+		[
+			SAssignNew(InputText, SEditableText)
+			.Text(FText::FromString(InSavedText))
+			.HintText(LOCTEXT("CommandValueTypeRowInputHintText", "Value..."))
+			.IsEnabled(this, &SConsoleVariablesEditorListValueInput::IsRowChecked)
+			.Visibility(Item.Pin()->GetCommandInfo().Pin()->ObjectType ==
+				FConsoleVariablesEditorCommandInfo::EConsoleObjectType::Command ?
+				EVisibility::Visible : EVisibility::Collapsed)
+		]
+
+		+SHorizontalBox::Slot()
+		.VAlign(VAlign_Center)
+		[
+			SAssignNew(InputWidget, SButton)
+			.OnClicked_Lambda([this] ()
+			{
+				check (Item.IsValid());
+				
+				const FString ValueAsString = InputText->GetText().ToString();
+				
+				if (const TSharedPtr<FConsoleVariablesEditorListRow> PinnedItem = Item.Pin())
+				{
+					// Only execute if this row is a non-value command OR
+					// it's a value-based command whose cached value differs from the input value
+					if (PinnedItem->GetCommandInfo().Pin()->ObjectType ==
+						FConsoleVariablesEditorCommandInfo::EConsoleObjectType::NullObject ||
+						(PinnedItem->GetCommandInfo().Pin()->ObjectType ==
+						FConsoleVariablesEditorCommandInfo::EConsoleObjectType::Command &&
+						!PinnedItem->GetCachedValue().Equals(ValueAsString)))
+					{
+						PinnedItem->GetCommandInfo().Pin()->ExecuteCommand(ValueAsString);
+
+						FConsoleVariablesEditorModule::Get().SendMultiUserConsoleVariableChange(
+							PinnedItem->GetCommandInfo().Pin()->Command, ValueAsString);
+
+						PinnedItem->SetCachedValue(ValueAsString);
+					}
+
+					return FReply::Handled();
+				}
+
+				return FReply::Unhandled();
+			})
+			.IsEnabled(this, &SConsoleVariablesEditorListValueInput::IsRowChecked)
+			[
+				SNew(STextBlock)
+				.Justification(ETextJustify::Center)
+				.Text(LOCTEXT("ConsoleCommandExecutionButtonText","Execute"))
+			]
+		]
+	];
+
+	Item.Pin()->SetCachedValue(GetInputValueAsString());
+}
+
+void SConsoleVariablesEditorListValueInput_Command::SetInputValue(const FString& InValueAsString)
+{
+	InputText->SetText(FText::FromString(InValueAsString));
+}
+
+FString SConsoleVariablesEditorListValueInput_Command::GetInputValueAsString()
+{
+	return GetInputValue();
+}
+
+FString SConsoleVariablesEditorListValueInput_Command::GetInputValue()
+{
+	return InputText->GetText().ToString();
+}
+
 #undef LOCTEXT_NAMESPACE
diff --git a/Engine/Plugins/Experimental/ConsoleVariables/Source/ConsoleVariablesEditor/Private/Views/List/SConsoleVariablesEditorListValueInput.h b/Engine/Plugins/Experimental/ConsoleVariables/Source/ConsoleVariablesEditor/Private/Views/List/SConsoleVariablesEditorListValueInput.h
--- a/Engine/Plugins/Experimental/ConsoleVariables/Source/ConsoleVariablesEditor/Private/Views/List/SConsoleVariablesEditorListValueInput.h
+++ b/Engine/Plugins/Experimental/ConsoleVariables/Source/ConsoleVariablesEditor/Private/Views/List/SConsoleVariablesEditorListValueInput.h
@@ -126,3 +126,26 @@ private:
 	TSharedPtr<SButton> InputWidget;
 	TSharedPtr<STextBlock> ButtonText;
 };
+
+class SConsoleVariablesEditorListValueInput_Command : public SConsoleVariablesEditorListValueInput
+{
+public:
+	
+	SLATE_BEGIN_ARGS(SConsoleVariablesEditorListValueInput_Command)
+	{}
+
+	SLATE_END_ARGS()
+
+	void Construct(
+		const FArguments& InArgs, const TWeakPtr<FConsoleVariablesEditorListRow> InRow, const FString& InSavedText);
+	
+	virtual void SetInputValue(const FString& InValueAsString) override;
+	virtual FString GetInputValueAsString() override;
+
+	FString GetInputValue();
+
+private:
+	
+	TSharedPtr<SButton> InputWidget;
+	TSharedPtr<SEditableText> InputText;
+};
diff --git a/Engine/Plugins/Experimental/ConsoleVariables/Source/ConsoleVariablesEditor/Private/Views/MainPanel/ConsoleVariablesEditorMainPanel.cpp b/Engine/Plugins/Experimental/ConsoleVariables/Source/ConsoleVariablesEditor/Private/Views/MainPanel/ConsoleVariablesEditorMainPanel.cpp
--- a/Engine/Plugins/Experimental/ConsoleVariables/Source/ConsoleVariablesEditor/Private/Views/MainPanel/ConsoleVariablesEditorMainPanel.cpp
+++ b/Engine/Plugins/Experimental/ConsoleVariables/Source/ConsoleVariablesEditor/Private/Views/MainPanel/ConsoleVariablesEditorMainPanel.cpp
@@ -51,7 +51,7 @@ TObjectPtr<UConsoleVariablesAsset> FConsoleVariablesEditorMainPanel::GetEditingA
 	return GetConsoleVariablesModule().GetEditingAsset();
 }
 
-void FConsoleVariablesEditorMainPanel::AddConsoleVariable(
+void FConsoleVariablesEditorMainPanel::AddConsoleObjectToPreset(
 	const FString& InConsoleCommand, const FString& InValue, const bool bScrollToNewRow) const
 {
 	const TWeakObjectPtr<UConsoleVariablesAsset> EditingAsset = GetEditingAsset();
@@ -60,7 +60,7 @@ void FConsoleVariablesEditorMainPanel::AddConsoleVariable(
 	{
 		UConsoleVariablesAsset* Asset = EditingAsset.Get();
 
-		Asset->AddOrSetConsoleVariableSavedData(
+		Asset->AddOrSetConsoleObjectSavedData(
 			{
 				InConsoleCommand,
 				InValue,
@@ -160,7 +160,8 @@ void FConsoleVariablesEditorMainPanel::ImportPreset(const FAssetData& InPresetAs
 	}
 }
 
-bool FConsoleVariablesEditorMainPanel::ImportPreset_Impl(const FAssetData& InPresetAsset, const TObjectPtr<UConsoleVariablesAsset> EditingAsset)
+bool FConsoleVariablesEditorMainPanel::ImportPreset_Impl(
+	const FAssetData& InPresetAsset, const TObjectPtr<UConsoleVariablesAsset> EditingAsset)
 {
 	if (UConsoleVariablesAsset* Preset = CastChecked<UConsoleVariablesAsset>(InPresetAsset.GetAsset()))
 	{
diff --git a/Engine/Plugins/Experimental/ConsoleVariables/Source/ConsoleVariablesEditor/Private/Views/MainPanel/ConsoleVariablesEditorMainPanel.h b/Engine/Plugins/Experimental/ConsoleVariables/Source/ConsoleVariablesEditor/Private/Views/MainPanel/ConsoleVariablesEditorMainPanel.h
--- a/Engine/Plugins/Experimental/ConsoleVariables/Source/ConsoleVariablesEditor/Private/Views/MainPanel/ConsoleVariablesEditorMainPanel.h
+++ b/Engine/Plugins/Experimental/ConsoleVariables/Source/ConsoleVariablesEditor/Private/Views/MainPanel/ConsoleVariablesEditorMainPanel.h
@@ -21,7 +21,9 @@ public:
 	static FConsoleVariablesEditorModule& GetConsoleVariablesModule();
 	static TObjectPtr<UConsoleVariablesAsset> GetEditingAsset();
 
-	void AddConsoleVariable(const FString& InConsoleCommand, const FString& InValue, const bool bScrollToNewRow = false) const;
+	void AddConsoleObjectToPreset(
+		const FString& InConsoleCommand, const FString& InValue,
+		const bool bScrollToNewRow = false) const;
 
 	/** Regenerate the list items and refresh the list. Call when adding or removing variables. */
 	void RebuildList(const FString& InConsoleCommandToScrollTo = "") const;
diff --git a/Engine/Plugins/Experimental/ConsoleVariables/Source/ConsoleVariablesEditor/Private/Views/MainPanel/SConsoleVariablesEditorCustomConsoleInputBox.cpp b/Engine/Plugins/Experimental/ConsoleVariables/Source/ConsoleVariablesEditor/Private/Views/MainPanel/SConsoleVariablesEditorCustomConsoleInputBox.cpp
--- a/Engine/Plugins/Experimental/ConsoleVariables/Source/ConsoleVariablesEditor/Private/Views/MainPanel/SConsoleVariablesEditorCustomConsoleInputBox.cpp
+++ b/Engine/Plugins/Experimental/ConsoleVariables/Source/ConsoleVariablesEditor/Private/Views/MainPanel/SConsoleVariablesEditorCustomConsoleInputBox.cpp
@@ -79,12 +79,11 @@ void SConsoleVariablesEditorCustomConsoleInputBox::Construct(
 						TSharedPtr<IToolTip> ToolTip;
 
 						if (const TWeakPtr<FConsoleVariablesEditorCommandInfo> MatchingCommand =
-							ConsoleVariablesEditorModule.FindCommandInfoByName(SanitizedText);
-							MatchingCommand.IsValid() && MatchingCommand.Pin()->ConsoleVariablePtr)
+							ConsoleVariablesEditorModule.FindCommandInfoByName(SanitizedText); MatchingCommand.IsValid())
 						{
 							ToolTip = SConsoleVariablesEditorTooltipWidget::MakeTooltip(
 								SanitizedText,
-								MatchingCommand.Pin()->ConsoleVariablePtr->GetHelp());
+								MatchingCommand.Pin()->GetHelpText());
 						}
 
 						return
@@ -221,11 +220,6 @@ void SConsoleVariablesEditorCustomConsoleInputBox::OnTextChanged(const FText& In
 				return;
 			}
 
-			if (!CVar->AsVariable())
-			{
-				return;
-			}
-
 			AutoCompleteList.Add(Name);
 		};
 
diff --git a/Engine/Plugins/Experimental/ConsoleVariables/Source/ConsoleVariablesEditor/Private/Views/MainPanel/SConsoleVariablesEditorMainPanel.cpp b/Engine/Plugins/Experimental/ConsoleVariables/Source/ConsoleVariablesEditor/Private/Views/MainPanel/SConsoleVariablesEditorMainPanel.cpp
--- a/Engine/Plugins/Experimental/ConsoleVariables/Source/ConsoleVariablesEditor/Private/Views/MainPanel/SConsoleVariablesEditorMainPanel.cpp
+++ b/Engine/Plugins/Experimental/ConsoleVariables/Source/ConsoleVariablesEditor/Private/Views/MainPanel/SConsoleVariablesEditorMainPanel.cpp
@@ -62,32 +62,58 @@ void SConsoleVariablesEditorMainPanel::Construct(
 
 SConsoleVariablesEditorMainPanel::~SConsoleVariablesEditorMainPanel()
 {
-
+	MainPanel.Reset();
+	ToolbarHBox.Reset();
+	ConcertButtonPtr.Reset();
+	MultiUserDetailsBox.Reset();
 }
 
 FReply SConsoleVariablesEditorMainPanel::ValidateConsoleInput(const FText& CommittedText)
 {
-	FString CommandString = CommittedText.ToString().TrimStartAndEnd();
+	const FString CommandString = CommittedText.ToString().TrimStartAndEnd();
+	FString CommandKey;
 	FString ValueString;
 	
 	if (CommandString.Contains(" "))
 	{
-		CommandString.Split(TEXT(" "), &CommandString, &ValueString);
+		CommandString.Split(TEXT(" "), &CommandKey, &ValueString);
 	}
 
-	if (IConsoleVariable* AsVariable = IConsoleManager::Get().FindConsoleVariable(*CommandString))
+	if (IConsoleObject* ConsoleObject = IConsoleManager::Get().FindConsoleObject(*CommandKey)) 
 	{
-		MainPanel.Pin()->AddConsoleVariable(
-			CommandString, ValueString.IsEmpty() ? AsVariable->GetString() : ValueString, true);
+		FString InValue = ValueString;
+
+		if (InValue.IsEmpty())
+		{
+			if (const IConsoleVariable* AsVariable = ConsoleObject->AsVariable())
+			{
+				InValue = AsVariable->GetString();
+			}
+		}
+		
+		MainPanel.Pin()->AddConsoleObjectToPreset(
+			CommandKey,
+			InValue,
+			true
+		);
 	}
 	else if (CommandString.IsEmpty())
 	{
-		UE_LOG(LogConsoleVariablesEditor, Warning, TEXT("hs: Input is blank."), __FUNCTION__);
+		UE_LOG(LogConsoleVariablesEditor, Warning, TEXT("%hs: Input is blank."), __FUNCTION__);
+	}
+	// Try to execute the whole command. Some commands are not registered, but are parsed externally so they won't be found by IConsoleManager::FindConsoleObject
+	else if (GEngine->Exec(FConsoleVariablesEditorCommandInfo::GetCurrentWorld(), *CommandString))
+	{
+		MainPanel.Pin()->AddConsoleObjectToPreset(
+			CommandString,
+			"",
+			true
+			);
 	}
 	else
 	{
 		UE_LOG(LogConsoleVariablesEditor, Warning,
-			TEXT("%hs: Input %s is not a recognized console variable. Note that only console variables are supported, not console commands."),
+			TEXT("%hs: Input '%s' is not a recognized console variable or command."),
 			__FUNCTION__, *CommandString);
 	}
 
diff --git a/Engine/Plugins/Experimental/ConsoleVariables/Source/ConsoleVariablesEditor/Public/ConsoleVariablesAsset.h b/Engine/Plugins/Experimental/ConsoleVariables/Source/ConsoleVariablesEditor/Public/ConsoleVariablesAsset.h
--- a/Engine/Plugins/Experimental/ConsoleVariables/Source/ConsoleVariablesEditor/Public/ConsoleVariablesAsset.h
+++ b/Engine/Plugins/Experimental/ConsoleVariables/Source/ConsoleVariablesEditor/Public/ConsoleVariablesAsset.h
@@ -2,13 +2,13 @@
 
 #pragma once
 
+#include "ConsoleVariablesEditorCommandInfo.h"
+
 #include "CoreMinimal.h"
 #include "Widgets/Input/SCheckBox.h"
 
 #include "ConsoleVariablesAsset.generated.h"
 
-struct FConsoleVariablesEditorCommandInfo;
-
 /** Data that will be serialized with this asset */
 USTRUCT()
 struct FConsoleVariablesEditorAssetSaveData
@@ -27,7 +27,7 @@ struct FConsoleVariablesEditorAssetSaveData
 	FString CommandValueAsString;
 
 	UPROPERTY()
-	// If Undetermined, we can assume this data was not previously saved
+	/** If Undetermined, we can assume this data was not previously saved */
 	ECheckBoxState CheckedState = ECheckBoxState::Undetermined;
 };
 
@@ -66,7 +66,7 @@ public:
 	bool FindSavedDataByCommandString(const FString& InCommandString, FConsoleVariablesEditorAssetSaveData& OutValue) const;
 
 	/** Set the value of a saved console variable if the name matches; add a new console variable to the list if a match is not found. */
-	void AddOrSetConsoleVariableSavedData(const FConsoleVariablesEditorAssetSaveData& InData);
+	void AddOrSetConsoleObjectSavedData(const FConsoleVariablesEditorAssetSaveData& InData);
 
 	/** Returns true if the element was found and successfully removed. */
 	bool RemoveConsoleVariable(const FString& InCommandString);
diff --git a/Engine/Plugins/Experimental/ConsoleVariables/Source/ConsoleVariablesEditor/Public/ConsoleVariablesEditorProjectSettings.h b/Engine/Plugins/Experimental/ConsoleVariables/Source/ConsoleVariablesEditor/Public/ConsoleVariablesEditorProjectSettings.h
--- a/Engine/Plugins/Experimental/ConsoleVariables/Source/ConsoleVariablesEditor/Public/ConsoleVariablesEditorProjectSettings.h
+++ b/Engine/Plugins/Experimental/ConsoleVariables/Source/ConsoleVariablesEditor/Public/ConsoleVariablesEditorProjectSettings.h
@@ -22,14 +22,19 @@ public:
 		bAddAllChangedConsoleVariablesToCurrentPreset = true;
 	}
 
-	/** If true, any console variable changes will be added to the current preset as long as the plugin is loaded. */
+	/**
+	 *When variables are changed outside the Console Variables Editor, this option will add the variables to the current preset.
+	 *Does not apply to console commands like 'r.SetNearClipPlane' or 'stat fps'
+	 */
 	UPROPERTY(Config, EditAnywhere, Category="Console Variables")
 	bool bAddAllChangedConsoleVariablesToCurrentPreset;
 
-	/** When a row is unchecked, its associated variable's value will be set to the value recorded when the plugin was loaded.
+	/**
+	 *When a row is unchecked, its associated variable's value will be set to the value recorded when the plugin was loaded.
 	 *The value displayed to the user can be configured with this setting, but will not affect the actual applied value.
 	 *ShowCurrentValue displays the actual value currently applied to the variable.
-	 *ShowLastEnteredValue displays the value that will be applied when the row is checked. */
+	 *ShowLastEnteredValue displays the value that will be applied when the row is checked.
+	 */
 	UPROPERTY(Config, EditAnywhere, Category="Console Variables")
 	EConsoleVariablesEditorRowDisplayType UncheckedRowDisplayType;
 };
