Console Variables Editor:

-Fix for: Invalid values entered when item is added will have stuck reset arrow

#[fyi] Jason.Walter

diff --git a/Engine/Plugins/Experimental/ConsoleVariables/Source/ConsoleVariablesEditor/Private/Views/List/ConsoleVariablesEditorListRow.cpp b/Engine/Plugins/Experimental/ConsoleVariables/Source/ConsoleVariablesEditor/Private/Views/List/ConsoleVariablesEditorListRow.cpp
--- a/Engine/Plugins/Experimental/ConsoleVariables/Source/ConsoleVariablesEditor/Private/Views/List/ConsoleVariablesEditorListRow.cpp
+++ b/Engine/Plugins/Experimental/ConsoleVariables/Source/ConsoleVariablesEditor/Private/Views/List/ConsoleVariablesEditorListRow.cpp
@@ -280,3 +280,8 @@ const FString& FConsoleVariablesEditorListRow::GetPresetValue() const
 {
 	return PresetValue;
 }
+
+void FConsoleVariablesEditorListRow::SetPresetValue(const FString& InPresetValue)
+{
+	PresetValue = InPresetValue;
+}
diff --git a/Engine/Plugins/Experimental/ConsoleVariables/Source/ConsoleVariablesEditor/Private/Views/List/ConsoleVariablesEditorListRow.h b/Engine/Plugins/Experimental/ConsoleVariables/Source/ConsoleVariablesEditor/Private/Views/List/ConsoleVariablesEditorListRow.h
--- a/Engine/Plugins/Experimental/ConsoleVariables/Source/ConsoleVariablesEditor/Private/Views/List/ConsoleVariablesEditorListRow.h
+++ b/Engine/Plugins/Experimental/ConsoleVariables/Source/ConsoleVariablesEditor/Private/Views/List/ConsoleVariablesEditorListRow.h
@@ -72,6 +72,7 @@ struct FConsoleVariablesEditorListRow final : TSharedFromThis<FConsoleVariablesE
 	void ResetToStartupValueAndSource() const;
 	
 	[[nodiscard]] const FString& GetPresetValue() const;
+	void SetPresetValue(const FString& InPresetValue);
 
 	/* If bMatchAnyTokens is false, only nodes that match all terms will be returned. */
 	bool MatchSearchTokensToSearchTerms(const TArray<FString> InTokens, const bool bMatchAnyTokens = false);
diff --git a/Engine/Plugins/Experimental/ConsoleVariables/Source/ConsoleVariablesEditor/Private/Views/List/SConsoleVariablesEditorListValueInput.cpp b/Engine/Plugins/Experimental/ConsoleVariables/Source/ConsoleVariablesEditor/Private/Views/List/SConsoleVariablesEditorListValueInput.cpp
--- a/Engine/Plugins/Experimental/ConsoleVariables/Source/ConsoleVariablesEditor/Private/Views/List/SConsoleVariablesEditorListValueInput.cpp
+++ b/Engine/Plugins/Experimental/ConsoleVariables/Source/ConsoleVariablesEditor/Private/Views/List/SConsoleVariablesEditorListValueInput.cpp
@@ -188,6 +188,10 @@ void SConsoleVariablesEditorListValueInput_Int::Construct(const FArguments& InAr
 	{
 		InputWidget->SetMinSliderValue(0);
 		InputWidget->SetMaxSliderValue(2);
+
+		const int32 PresetValue = FCString::Atoi(*Item.Pin()->GetPresetValue());
+
+		Item.Pin()->SetPresetValue(FString::FromInt(FMath::Clamp(PresetValue, 0, 2)));
 	}
 	
 	Item.Pin()->SetCachedValue(GetInputValueAsString());
