Console Variables Editor:

-Fix for: Tooltips for Preset menu elements appear below entire Preset menu

#[fyi] Jason.Walter

diff --git a/Engine/Plugins/Experimental/ConsoleVariables/Source/ConsoleVariablesEditor/Private/Views/MainPanel/SConsoleVariablesEditorMainPanel.cpp b/Engine/Plugins/Experimental/ConsoleVariables/Source/ConsoleVariablesEditor/Private/Views/MainPanel/SConsoleVariablesEditorMainPanel.cpp
--- a/Engine/Plugins/Experimental/ConsoleVariables/Source/ConsoleVariablesEditor/Private/Views/MainPanel/SConsoleVariablesEditorMainPanel.cpp
+++ b/Engine/Plugins/Experimental/ConsoleVariables/Source/ConsoleVariablesEditor/Private/Views/MainPanel/SConsoleVariablesEditorMainPanel.cpp
@@ -78,22 +78,24 @@ FReply SConsoleVariablesEditorMainPanel::ValidateConsoleInput(const FText& Commi
 	{
 		CommandString.Split(TEXT(" "), &CommandKey, &ValueString);
 	}
+	else
+	{
+		CommandKey = CommandString;
+	}
 
 	if (IConsoleObject* ConsoleObject = IConsoleManager::Get().FindConsoleObject(*CommandKey)) 
 	{
-		FString InValue = ValueString;
-
-		if (InValue.IsEmpty())
+		if (ValueString.IsEmpty())
 		{
 			if (const IConsoleVariable* AsVariable = ConsoleObject->AsVariable())
 			{
-				InValue = AsVariable->GetString();
+				ValueString = AsVariable->GetString();
 			}
 		}
 		
 		MainPanel.Pin()->AddConsoleObjectToPreset(
 			CommandKey,
-			InValue,
+			ValueString,
 			true
 		);
 	}
@@ -238,7 +240,7 @@ void SConsoleVariablesEditorMainPanel::CreateConcertButtonIfNeeded()
 
 TSharedRef<SWidget> SConsoleVariablesEditorMainPanel::OnGeneratePresetsMenu()
 {
-	FMenuBuilder MenuBuilder(true, nullptr);
+	FMenuBuilder MenuBuilder(true, NULL);
 
 	IContentBrowserSingleton& ContentBrowser =
 		FModuleManager::LoadModuleChecked<FContentBrowserModule>("ContentBrowser").Get();
@@ -250,6 +252,7 @@ TSharedRef<SWidget> SConsoleVariablesEditorMainPanel::OnGeneratePresetsMenu()
 
 	FUIAction PresetNameAction = FUIAction();
 	PresetNameAction.CanExecuteAction = FCanExecuteAction::CreateLambda([]() { return false; });
+	
 	MenuBuilder.AddMenuEntry(
 		LoadedPresetName,
 		LoadedPresetName,
@@ -265,14 +268,18 @@ TSharedRef<SWidget> SConsoleVariablesEditorMainPanel::OnGeneratePresetsMenu()
 		LOCTEXT("SavePreset_Text", "Save Preset"),
 		LOCTEXT("SavePreset_Tooltip", "Save the current preset if one has been loaded. Otherwise, the Save As dialog will be opened."),
 		FSlateIcon(FAppStyle::Get().GetStyleSetName(), "AssetEditor.SaveAsset"),
-		FUIAction(FExecuteAction::CreateRaw(MainPanel.Pin().Get(), &FConsoleVariablesEditorMainPanel::SavePreset))
+		FUIAction(FExecuteAction::CreateRaw(MainPanel.Pin().Get(), &FConsoleVariablesEditorMainPanel::SavePreset)),
+		NAME_None,
+		EUserInterfaceActionType::Button
 	);
 
 	MenuBuilder.AddMenuEntry(
 		LOCTEXT("SavePresetAs_Text", "Save Preset As"),
 		LOCTEXT("SavePresetAs_Tooltip", "Save the current configuration as a new preset that can be shared between multiple jobs, or imported later as the base of a new configuration."),
 		FSlateIcon(FAppStyle::Get().GetStyleSetName(), "AssetEditor.SaveAssetAs"),
-		FUIAction(FExecuteAction::CreateRaw(MainPanel.Pin().Get(), &FConsoleVariablesEditorMainPanel::SavePresetAs))
+		FUIAction(FExecuteAction::CreateRaw(MainPanel.Pin().Get(), &FConsoleVariablesEditorMainPanel::SavePresetAs)),
+		NAME_None,
+		EUserInterfaceActionType::Button
 	);
 
 	FAssetPickerConfig AssetPickerConfig;
@@ -311,7 +318,17 @@ TSharedRef<SWidget> SConsoleVariablesEditorMainPanel::OnGeneratePresetsMenu()
 	}
 	MenuBuilder.EndSection();
 
-	return MenuBuilder.MakeWidget();
+	TSharedRef<SWidget> Widget = MenuBuilder.MakeWidget();
+	FChildren* ChildWidgets = Widget->GetChildren();
+	for (int32 ChildItr = 0; ChildItr < ChildWidgets->Num(); ChildItr++)
+	{
+		TSharedRef<SWidget> Child = ChildWidgets->GetChildAt(ChildItr);
+
+		Child->EnableToolTipForceField(false);
+	}
+	Widget->EnableToolTipForceField(false);
+	
+	return Widget;
 }
 
 TSharedRef<SWidget> SConsoleVariablesEditorMainPanel::GetConcertDetailsWidget(UObject* InObject)
