Console Variables Editor:

-Slack-Style Search (Global Searches)
 -Added second transient preset that contains the global search matches to switch between modes
-Added Hover Widget to global search to add/remove item wrt the current preset
-Created configurable skiplist to filter out auto-populated changed variables
-Added Show Option "Show Only Modified"
 -Edited filters to have one of two types: MatchAny and MatchAll. In order for a row to pass the filter step, it must match at least one filter in the MatchAny group and all filters in the MatchAll group.
-Refactored visibility code: now a separate array of visible objects is passed into the tree view rather than collapsing existing tree items that should be invisible as doing that causes scrolling to malfunction in long lists

diff --git a/Engine/Plugins/Experimental/ConsoleVariables/Source/ConsoleVariablesEditor/Private/ConsoleVariablesEditorModule.cpp b/Engine/Plugins/Experimental/ConsoleVariables/Source/ConsoleVariablesEditor/Private/ConsoleVariablesEditorModule.cpp
--- a/Engine/Plugins/Experimental/ConsoleVariables/Source/ConsoleVariablesEditor/Private/ConsoleVariablesEditorModule.cpp
+++ b/Engine/Plugins/Experimental/ConsoleVariables/Source/ConsoleVariablesEditor/Private/ConsoleVariablesEditorModule.cpp
@@ -9,6 +9,7 @@
 #include "MultiUser/ConsoleVariableSyncData.h"
 #include "Views/MainPanel/ConsoleVariablesEditorMainPanel.h"
 
+#include "Algo/AnyOf.h"
 #include "Algo/Find.h"
 #include "Framework/Docking/TabManager.h"
 #include "ISettingsModule.h"
@@ -56,7 +57,7 @@ void FConsoleVariablesEditorModule::ShutdownModule()
 	}
 }
 
-void FConsoleVariablesEditorModule::OpenConsoleVariablesDialogWithAssetSelected(const FAssetData& InAssetData)
+void FConsoleVariablesEditorModule::OpenConsoleVariablesDialogWithAssetSelected(const FAssetData& InAssetData) const
 {
 	if (InAssetData.IsValid())
 	{
@@ -113,6 +114,26 @@ TWeakPtr<FConsoleVariablesEditorCommandInfo> FConsoleVariablesEditorModule::Find
 	return Match ? *Match : nullptr;
 }
 
+TArray<TWeakPtr<FConsoleVariablesEditorCommandInfo>> FConsoleVariablesEditorModule::FindCommandInfosMatchingTokens(
+	const TArray<FString>& InTokens, ESearchCase::Type InSearchCase)
+{
+	TArray<TWeakPtr<FConsoleVariablesEditorCommandInfo>> ReturnValue;
+	
+	for (const TSharedPtr<FConsoleVariablesEditorCommandInfo>& CommandInfo : ConsoleObjectsMasterReference)
+	{
+		if (Algo::AnyOf(InTokens,
+			[&CommandInfo, InSearchCase](const FString& Token)
+			{
+				return CommandInfo->Command.Contains(Token, InSearchCase);
+			}))
+		{
+			ReturnValue.Add(CommandInfo);
+		}
+	}
+
+	return ReturnValue;
+}
+
 TWeakPtr<FConsoleVariablesEditorCommandInfo> FConsoleVariablesEditorModule::FindCommandInfoByConsoleObjectReference(
 	IConsoleObject* InConsoleObjectReference)
 {
@@ -126,27 +147,66 @@ TWeakPtr<FConsoleVariablesEditorCommandInfo> FConsoleVariablesEditorModule::Find
 	return Match ? *Match : nullptr;
 }
 
-TObjectPtr<UConsoleVariablesAsset> FConsoleVariablesEditorModule::GetEditingAsset() const
+TObjectPtr<UConsoleVariablesAsset> FConsoleVariablesEditorModule::GetPresetAsset() const
 {
-	return EditingAsset.Get();
+	return EditingPresetAsset;
 }
 
-void FConsoleVariablesEditorModule::SetEditingAsset(const TObjectPtr<UConsoleVariablesAsset> InEditingAsset)
+TObjectPtr<UConsoleVariablesAsset> FConsoleVariablesEditorModule::GetGlobalSearchAsset() const
 {
-	EditingAsset = InEditingAsset;
+	return EditingGlobalSearchAsset;
 }
 
-void FConsoleVariablesEditorModule::SendMultiUserConsoleVariableChange(const FString& InVariableName, const FString& InValueAsString)
+bool FConsoleVariablesEditorModule::PopulateGlobalSearchAssetWithVariablesMatchingTokens(const TArray<FString>& InTokens)
+{
+	// Remove existing commands
+	TArray<FConsoleVariablesEditorAssetSaveData> NullList;
+	EditingGlobalSearchAsset->ReplaceSavedCommands(NullList);
+	
+	for (const TWeakPtr<FConsoleVariablesEditorCommandInfo>& CommandInfo : FindCommandInfosMatchingTokens(InTokens))
+	{
+		EditingGlobalSearchAsset->AddOrSetConsoleObjectSavedData(
+			{
+				CommandInfo.Pin()->Command,
+				CommandInfo.Pin()->GetConsoleVariablePtr() ?
+					CommandInfo.Pin()->GetConsoleVariablePtr()->GetString() : "",
+				ECheckBoxState::Checked
+			}
+		);
+	}
+
+	return EditingGlobalSearchAsset->GetSavedCommandsCount() > 0;
+}
+
+void FConsoleVariablesEditorModule::SendMultiUserConsoleVariableChange(const FString& InVariableName, const FString& InValueAsString) const
 {
 	MainPanel->GetMultiUserManager().SendConsoleVariableChange(InVariableName, InValueAsString);
 }
 
+void FConsoleVariablesEditorModule::OnRemoteCvarChanged(const FString InName, const FString InValue)
+{
+	UE_LOG(LogConsoleVariablesEditor, Display, TEXT("Remote set console variable %s = %s"), *InName, *InValue);
+
+	if (GetDefault<UConcertCVarSynchronization>()->bSyncCVarTransactions)
+	{
+		if (const TWeakPtr<FConsoleVariablesEditorCommandInfo> CommandInfo =
+			FindCommandInfoByName(InName); CommandInfo.IsValid())
+		{
+			if (CommandInfo.Pin()->IsCurrentValueDifferentFromInputValue(InValue))
+			{
+				GEngine->Exec(FConsoleVariablesEditorCommandInfo::GetCurrentWorld(),
+					*FString::Printf(TEXT("%s %s"), *InName, *InValue));
+			}
+		}
+	}
+}
+
 void FConsoleVariablesEditorModule::OnFEngineLoopInitComplete()
 {
 	RegisterMenuItem();
 	RegisterProjectSettings();
 	QueryAndBeginTrackingConsoleVariables();
-	AllocateTransientPreset();
+	CreateEditingPresets();
 	
 	MainPanel = MakeShared<FConsoleVariablesEditorMainPanel>();
 }
@@ -164,7 +224,7 @@ void FConsoleVariablesEditorModule::RegisterMenuItem()
 	BrowserSpawnerEntry.SetGroup(WorkspaceMenu::GetMenuStructure().GetLevelEditorCategory());
 }
 
-void FConsoleVariablesEditorModule::RegisterProjectSettings()
+void FConsoleVariablesEditorModule::RegisterProjectSettings() const
 {
 	ISettingsModule& SettingsModule = FModuleManager::LoadModuleChecked<ISettingsModule>("Settings");
 	{
@@ -180,7 +240,7 @@ void FConsoleVariablesEditorModule::RegisterProjectSettings()
 
 void FConsoleVariablesEditorModule::OnConsoleVariableChanged(IConsoleVariable* ChangedVariable)
 {
-	check(EditingAsset);
+	check(EditingPresetAsset);
 
 	if (const TWeakPtr<FConsoleVariablesEditorCommandInfo> CommandInfo =
 		FindCommandInfoByConsoleObjectReference(ChangedVariable); CommandInfo.IsValid())
@@ -189,20 +249,28 @@ void FConsoleVariablesEditorModule::OnConsoleVariableChanged(IConsoleVariable* C
 		const FString& Key = PinnedCommand->Command;
 		
 		FConsoleVariablesEditorAssetSaveData FoundData;
-		const bool bIsVariableCurrentlyTracked = EditingAsset->FindSavedDataByCommandString(Key, FoundData);
+		const bool bIsVariableCurrentlyTracked = EditingPresetAsset->FindSavedDataByCommandString(Key, FoundData);
+
+		const UConsoleVariablesEditorProjectSettings* Settings = GetDefault<UConsoleVariablesEditorProjectSettings>();
+		check(Settings);
 		
 		if (!bIsVariableCurrentlyTracked)
 		{
 			// If not yet tracked and we want to track variable changes from outside the dialogue,
 			// Check if the changed value differs from the startup value before tracking it
-			if (GetMutableDefault<UConsoleVariablesEditorProjectSettings>()->bAddAllChangedConsoleVariablesToCurrentPreset
-				&& PinnedCommand->IsCurrentValueDifferentFromInputValue(PinnedCommand->StartupValueAsString))
+			if (Settings->bAddAllChangedConsoleVariablesToCurrentPreset &&
+				!Settings->ChangedConsoleVariableSkipList.Contains(Key) && 
+				PinnedCommand->IsCurrentValueDifferentFromInputValue(PinnedCommand->StartupValueAsString))
 			{
 				if (MainPanel.IsValid())
 				{
 					MainPanel->AddConsoleObjectToPreset(
 						Key,
-						ChangedVariable->GetString(),
+						// If we're not in preset mode then pass empty value
+						// This forces the row to get the current value at the time it's generated
+						MainPanel->GetEditorListMode() ==
+						FConsoleVariablesEditorList::EConsoleVariablesEditorListMode::Preset ?
+						ChangedVariable->GetString() : "",
 						true
 					);
 				}
@@ -224,9 +292,9 @@ void FConsoleVariablesEditorModule::OnConsoleVariableChanged(IConsoleVariable* C
 
 void FConsoleVariablesEditorModule::OnDetectConsoleObjectUnregistered(FString CommandName)
 {
-	check(EditingAsset);
+	check(EditingPresetAsset);
 
-	EditingAsset->RemoveConsoleVariable(CommandName);
+	EditingPresetAsset->RemoveConsoleVariable(CommandName);
 
 	if (MainPanel.IsValid())
 	{
@@ -240,37 +308,23 @@ void FConsoleVariablesEditorModule::OnDetectConsoleObjectUnregistered(FString Co
 	}
 }
 
-void FConsoleVariablesEditorModule::OnRemoteCvarChanged(const FString InName, const FString InValue)
+TObjectPtr<UConsoleVariablesAsset> FConsoleVariablesEditorModule::AllocateTransientPreset(const FName DesiredName) const
 {
-	UE_LOG(LogConsoleVariablesEditor, Display, TEXT("Remote set console variable %s = %s"), *InName, *InValue);
+	const FString PackageName = "/Temp/ConsoleVariablesEditor/PendingConsoleVariablesPresets";
 
-	if (GetMutableDefault<UConcertCVarSynchronization>()->bSyncCVarTransactions)
-	{
-		if (const TWeakPtr<FConsoleVariablesEditorCommandInfo> CommandInfo =
-			FindCommandInfoByName(InName); CommandInfo.IsValid())
-		{
-			if (CommandInfo.Pin()->IsCurrentValueDifferentFromInputValue(InValue))
-			{
-				GEngine->Exec(FConsoleVariablesEditorCommandInfo::GetCurrentWorld(),
-					*FString::Printf(TEXT("%s %s"), *InName, *InValue));
-			}
-		}
-	}
-}
-
-TObjectPtr<UConsoleVariablesAsset> FConsoleVariablesEditorModule::AllocateTransientPreset()
-{
-	static const TCHAR* PackageName = TEXT("/Temp/ConsoleVariablesEditor/PendingConsoleVariablesCollections");
-
-	static FName DesiredName = "PendingConsoleVariablesCollection";
-
-	UPackage* NewPackage = CreatePackage(PackageName);
+	UPackage* NewPackage = CreatePackage(*PackageName);
 	NewPackage->SetFlags(RF_Transient);
 	NewPackage->AddToRoot();
 
-	EditingAsset = NewObject<UConsoleVariablesAsset>(NewPackage, DesiredName, RF_Transient | RF_Transactional | RF_Standalone);
+	return NewObject<UConsoleVariablesAsset>(
+		NewPackage, DesiredName, RF_Transient | RF_Transactional | RF_Standalone);
+}
 
-	return EditingAsset;
+void FConsoleVariablesEditorModule::CreateEditingPresets()
+{
+	EditingPresetAsset = AllocateTransientPreset("ConsoleVariablesPreset_PendingPreset");
+	
+	EditingGlobalSearchAsset = AllocateTransientPreset("ConsoleVariablesPreset_GlobalSearch");
 }
 
 TSharedRef<SDockTab> FConsoleVariablesEditorModule::SpawnMainPanelTab(const FSpawnTabArgs& Args)
diff --git a/Engine/Plugins/Experimental/ConsoleVariables/Source/ConsoleVariablesEditor/Private/ConsoleVariablesEditorModule.h b/Engine/Plugins/Experimental/ConsoleVariables/Source/ConsoleVariablesEditor/Private/ConsoleVariablesEditorModule.h
--- a/Engine/Plugins/Experimental/ConsoleVariables/Source/ConsoleVariablesEditor/Private/ConsoleVariablesEditorModule.h
+++ b/Engine/Plugins/Experimental/ConsoleVariables/Source/ConsoleVariablesEditor/Private/ConsoleVariablesEditorModule.h
@@ -28,7 +28,7 @@ public:
 	virtual void ShutdownModule() override;
 	//~ End IModuleInterface Interface
 
-	void OpenConsoleVariablesDialogWithAssetSelected(const FAssetData& InAssetData);
+	void OpenConsoleVariablesDialogWithAssetSelected(const FAssetData& InAssetData) const;
 
 	/** Find all console variables and cache their startup values */
 	void QueryAndBeginTrackingConsoleVariables();
@@ -42,6 +42,10 @@ public:
 	TWeakPtr<FConsoleVariablesEditorCommandInfo> FindCommandInfoByName(
 		const FString& NameToSearch, ESearchCase::Type InSearchCase = ESearchCase::IgnoreCase);
 
+	/** Find all tracked console variables matching a specific search query with optional case sensitivity. */
+	TArray<TWeakPtr<FConsoleVariablesEditorCommandInfo>> FindCommandInfosMatchingTokens(
+		const TArray<FString>& InTokens, ESearchCase::Type InSearchCase = ESearchCase::IgnoreCase);
+
 	/**
 	 *Find a tracked console variable by its console object reference.
 	 *Note that some commands do not have an associated console object (such as 'stat unit') and will not be found with this method.
@@ -50,10 +54,13 @@ public:
 	TWeakPtr<FConsoleVariablesEditorCommandInfo> FindCommandInfoByConsoleObjectReference(
 		IConsoleObject* InConsoleObjectReference);
 	
-	[[nodiscard]] TObjectPtr<UConsoleVariablesAsset> GetEditingAsset() const;
-	void SetEditingAsset(const TObjectPtr<UConsoleVariablesAsset> InEditingAsset);
+	[[nodiscard]] TObjectPtr<UConsoleVariablesAsset> GetPresetAsset() const;
+	[[nodiscard]] TObjectPtr<UConsoleVariablesAsset> GetGlobalSearchAsset() const;
+
+	/** Fills Global Search Asset's Saved Commands with variables matching the specified query. Returns false if no matches were found. */
+	bool PopulateGlobalSearchAssetWithVariablesMatchingTokens(const TArray<FString>& InTokens);
 
-	void SendMultiUserConsoleVariableChange(const FString& InVariableName, const FString& InValueAsString);
+	void SendMultiUserConsoleVariableChange(const FString& InVariableName, const FString& InValueAsString) const;
 	void OnRemoteCvarChanged(const FString InName, const FString InValue);
 	
 private:
@@ -61,25 +68,28 @@ private:
 	void OnFEngineLoopInitComplete();
 
 	void RegisterMenuItem();
-	void RegisterProjectSettings();
+	void RegisterProjectSettings() const;
 
 	void OnConsoleVariableChanged(IConsoleVariable* ChangedVariable);
 	/** In the event a console object is unregistered, this failsafe callback will clean up the associated list item and command info object. */
 	void OnDetectConsoleObjectUnregistered(FString CommandName);
 
-	TObjectPtr<UConsoleVariablesAsset> AllocateTransientPreset();
+	TObjectPtr<UConsoleVariablesAsset> AllocateTransientPreset(const FName DesiredName) const;
+	void CreateEditingPresets();
 
 	TSharedRef<SDockTab> SpawnMainPanelTab(const FSpawnTabArgs& Args);
 
-	void OpenConsoleVariablesEditor();
+	static void OpenConsoleVariablesEditor();
 	
 	static const FName ConsoleVariablesToolkitPanelTabId;
 
 	/* Lives for as long as the module is loaded. */
 	TSharedPtr<FConsoleVariablesEditorMainPanel> MainPanel;
 
-	// Transient preset that's being edited so we don't affect the reference asset unless we save it
-	TObjectPtr<UConsoleVariablesAsset> EditingAsset = nullptr;
+	/** Transient preset that's being edited so we don't affect the reference asset unless we save it */
+	TObjectPtr<UConsoleVariablesAsset> EditingPresetAsset = nullptr;
+	/** Transient preset that tracks variables that match the search criteria */
+	TObjectPtr<UConsoleVariablesAsset> EditingGlobalSearchAsset = nullptr;
 
 	/** All tracked variables and their default, startup, and current values */
 	TArray<TSharedPtr<FConsoleVariablesEditorCommandInfo>> ConsoleObjectsMasterReference;
diff --git a/Engine/Plugins/Experimental/ConsoleVariables/Source/ConsoleVariablesEditor/Private/Views/List/ConsoleVariablesEditorList.cpp b/Engine/Plugins/Experimental/ConsoleVariables/Source/ConsoleVariablesEditor/Private/Views/List/ConsoleVariablesEditorList.cpp
--- a/Engine/Plugins/Experimental/ConsoleVariables/Source/ConsoleVariablesEditor/Private/Views/List/ConsoleVariablesEditorList.cpp
+++ b/Engine/Plugins/Experimental/ConsoleVariables/Source/ConsoleVariablesEditor/Private/Views/List/ConsoleVariablesEditorList.cpp
@@ -6,14 +6,14 @@
 
 FConsoleVariablesEditorList::~FConsoleVariablesEditorList()
 {
-	
+	ListWidget.Reset();
 }
 
 TSharedRef<SWidget> FConsoleVariablesEditorList::GetOrCreateWidget()
 {
 	if (!ListWidget.IsValid())
 	{
-		SAssignNew(ListWidget, SConsoleVariablesEditorList);
+		SAssignNew(ListWidget, SConsoleVariablesEditorList, SharedThis(this));
 	}
 
 	return ListWidget.ToSharedRef();
diff --git a/Engine/Plugins/Experimental/ConsoleVariables/Source/ConsoleVariablesEditor/Private/Views/List/ConsoleVariablesEditorList.h b/Engine/Plugins/Experimental/ConsoleVariables/Source/ConsoleVariablesEditor/Private/Views/List/ConsoleVariablesEditorList.h
--- a/Engine/Plugins/Experimental/ConsoleVariables/Source/ConsoleVariablesEditor/Private/Views/List/ConsoleVariablesEditorList.h
+++ b/Engine/Plugins/Experimental/ConsoleVariables/Source/ConsoleVariablesEditor/Private/Views/List/ConsoleVariablesEditorList.h
@@ -14,12 +14,30 @@ class FConsoleVariablesEditorList : public TSharedFromThis<FConsoleVariablesEdit
 {
 public:
 
-	FConsoleVariablesEditorList(){};
+	enum class EConsoleVariablesEditorListMode : uint8
+	{
+		// We're displaying the cvars listed in the loaded (or default) preset
+		Preset = 1,
+		// We're displaying the cvars that match the criteria of the global search
+		GlobalSearch = 2
+	};
+
+	FConsoleVariablesEditorList() {}
 
 	~FConsoleVariablesEditorList();
 
 	TSharedRef<SWidget> GetOrCreateWidget();
 
+	EConsoleVariablesEditorListMode GetListMode() const
+	{
+		return CurrentListMode;
+	}
+
+	void SetListMode(EConsoleVariablesEditorListMode NewListMode)
+	{
+		CurrentListMode = NewListMode;;
+	}
+
 	/** Regenerate the list items and refresh the list. Call when adding or removing variables. */
 	void RebuildList(const FString& InConsoleCommandToScrollTo = "") const;
 
@@ -34,5 +52,6 @@ public:
 
 private:
 
+	EConsoleVariablesEditorListMode CurrentListMode = EConsoleVariablesEditorListMode::Preset;
 	TSharedPtr<SConsoleVariablesEditorList> ListWidget;
 };
diff --git a/Engine/Plugins/Experimental/ConsoleVariables/Source/ConsoleVariablesEditor/Private/Views/List/ConsoleVariablesEditorListFilters/ConsoleVariablesEditorListFilter_ModifiedVariables.h b/Engine/Plugins/Experimental/ConsoleVariables/Source/ConsoleVariablesEditor/Private/Views/List/ConsoleVariablesEditorListFilters/ConsoleVariablesEditorListFilter_ModifiedVariables.h
new file mode 100644
--- /dev/null
+++ b/Engine/Plugins/Experimental/ConsoleVariables/Source/ConsoleVariablesEditor/Private/Views/List/ConsoleVariablesEditorListFilters/ConsoleVariablesEditorListFilter_ModifiedVariables.h
@@ -0,0 +1,46 @@
+// Copyright Epic Games, Inc. All Rights Reserved.
+
+#pragma once
+
+#include "IConsoleVariablesEditorListFilter.h"
+
+#define LOCTEXT_NAMESPACE "ConsoleVariablesEditor"
+
+class ConsoleVariablesEditorListFilter_ModifiedVariables : public IConsoleVariablesEditorListFilter
+{
+public:
+
+	ConsoleVariablesEditorListFilter_ModifiedVariables()
+	{
+		// This filter should be off by default
+		SetFilterActive(false);
+		SetFilterMatchType(EConsoleVariablesEditorListFilterMatchType::MatchAll);
+	}
+
+	virtual FString GetFilterName() override
+	{
+		return "Show Only Modified Variables";
+	}
+
+	virtual FText GetFilterButtonLabel() override
+	{
+		return LOCTEXT("ShowOnlyModifiedVariablesFilter", "Show Only Modified");
+	}
+
+	virtual FText GetFilterButtonToolTip() override
+	{
+		return LOCTEXT("ShowSourceTextFilterTooltipFormat", "Show only rows that have a current value that differs from the preset value.");
+	}
+
+	virtual bool DoesItemPassFilter(const FConsoleVariablesEditorListRowPtr& InItem) override
+	{
+		if (InItem.IsValid())
+		{
+			return InItem->DoesCurrentValueDifferFromPresetValue();
+		}
+
+		return false;
+	}
+};
+
+#undef LOCTEXT_NAMESPACE
diff --git a/Engine/Plugins/Experimental/ConsoleVariables/Source/ConsoleVariablesEditor/Private/Views/List/ConsoleVariablesEditorListFilters/IConsoleVariablesEditorListFilter.h b/Engine/Plugins/Experimental/ConsoleVariables/Source/ConsoleVariablesEditor/Private/Views/List/ConsoleVariablesEditorListFilters/IConsoleVariablesEditorListFilter.h
--- a/Engine/Plugins/Experimental/ConsoleVariables/Source/ConsoleVariablesEditor/Private/Views/List/ConsoleVariablesEditorListFilters/IConsoleVariablesEditorListFilter.h
+++ b/Engine/Plugins/Experimental/ConsoleVariables/Source/ConsoleVariablesEditor/Private/Views/List/ConsoleVariablesEditorListFilters/IConsoleVariablesEditorListFilter.h
@@ -7,6 +7,14 @@
 class IConsoleVariablesEditorListFilter
 {
 public:
+
+	enum class EConsoleVariablesEditorListFilterMatchType
+	{
+		// Filters of this type are placed into a group in which if any of the group's filters pass, the row is considered passing
+		MatchAny,
+		// Filters of this type are placed into a group in which all of the group's filters must pass for the row to be considered passing
+		MatchAll
+	};
 	
 	virtual ~IConsoleVariablesEditorListFilter() = default;
 
@@ -32,6 +40,16 @@ public:
 		return bIsFilterActive;
 	}
 
+	void SetFilterMatchType(const EConsoleVariablesEditorListFilterMatchType InType)
+	{
+		FilterMatchType = InType;
+	}
+
+	EConsoleVariablesEditorListFilterMatchType GetFilterMatchType() const
+	{
+		return FilterMatchType;
+	}
+
 	virtual bool DoesItemPassFilter(const FConsoleVariablesEditorListRowPtr& InItem)
 	{
 		return false;
@@ -40,5 +58,6 @@ public:
 private:
 
 	bool bIsFilterActive = true;
+	EConsoleVariablesEditorListFilterMatchType FilterMatchType = EConsoleVariablesEditorListFilterMatchType::MatchAny;
 			
 };
diff --git a/Engine/Plugins/Experimental/ConsoleVariables/Source/ConsoleVariablesEditor/Private/Views/List/ConsoleVariablesEditorListRow.cpp b/Engine/Plugins/Experimental/ConsoleVariables/Source/ConsoleVariablesEditor/Private/Views/List/ConsoleVariablesEditorListRow.cpp
--- a/Engine/Plugins/Experimental/ConsoleVariables/Source/ConsoleVariablesEditor/Private/Views/List/ConsoleVariablesEditorListRow.cpp
+++ b/Engine/Plugins/Experimental/ConsoleVariables/Source/ConsoleVariablesEditor/Private/Views/List/ConsoleVariablesEditorListRow.cpp
@@ -6,6 +6,7 @@
 #include "ConsoleVariablesEditorModule.h"
 
 #include "Algo/AnyOf.h"
+#include "Views/List/ConsoleVariablesEditorList.h"
 #include "Views/List/SConsoleVariablesEditorList.h"
 
 FConsoleVariablesEditorListRow::~FConsoleVariablesEditorListRow()
@@ -83,7 +84,7 @@ void FConsoleVariablesEditorListRow::AddToChildRows(const FConsoleVariablesEdito
 }
 
 void FConsoleVariablesEditorListRow::InsertChildRowAtIndex(const FConsoleVariablesEditorListRowPtr& InRow,
-	const int32 AtIndex)
+                                                           const int32 AtIndex)
 {
 	ChildRows.Insert(InRow, AtIndex);
 }
@@ -98,6 +99,16 @@ void FConsoleVariablesEditorListRow::SetIsTreeViewItemExpanded(const bool bNewEx
 	bIsTreeViewItemExpanded = bNewExpanded;
 }
 
+bool FConsoleVariablesEditorListRow::GetShouldFlashOnScrollIntoView() const
+{
+	return bShouldFlashOnScrollIntoView;
+}
+
+void FConsoleVariablesEditorListRow::SetShouldFlashOnScrollIntoView(const bool bNewShouldFlashOnScrollIntoView)
+{
+	bShouldFlashOnScrollIntoView = bNewShouldFlashOnScrollIntoView;
+}
+
 bool FConsoleVariablesEditorListRow::GetShouldExpandAllChildren() const
 {
 	return bShouldExpandAllChildren;
@@ -114,6 +125,16 @@ void FConsoleVariablesEditorListRow::ResetToStartupValueAndSource() const
 	GetCommandInfo().Pin()->SetSourceFlag(GetCommandInfo().Pin()->StartupSource);
 }
 
+const FString& FConsoleVariablesEditorListRow::GetPresetValue() const
+{
+	return PresetValue;
+}
+
+void FConsoleVariablesEditorListRow::SetPresetValue(const FString& InPresetValue)
+{
+	PresetValue = InPresetValue;
+}
+
 bool FConsoleVariablesEditorListRow::MatchSearchTokensToSearchTerms(const TArray<FString> InTokens,
                                                                     const bool bMatchAnyTokens)
 {
@@ -123,7 +144,7 @@ bool FConsoleVariablesEditorListRow::MatchSearchTokensToSearchTerms(const TArray
 	if (!bMatchFound)
 	{
 		const TSharedPtr<FConsoleVariablesEditorCommandInfo> PinnedInfo = CommandInfo.Pin();
-		
+
 		FString SearchTerms = PinnedInfo->Command + PinnedInfo->GetSourceAsText().ToString();
 
 		if (const IConsoleVariable* AsVariable = PinnedInfo->GetConsoleVariablePtr())
@@ -133,10 +154,10 @@ bool FConsoleVariablesEditorListRow::MatchSearchTokensToSearchTerms(const TArray
 		}
 
 		bMatchFound = Algo::AnyOf(InTokens,
-			[&SearchTerms](const FString& Token)
-			{
-				return SearchTerms.Contains(Token);
-			});
+		                          [&SearchTerms](const FString& Token)
+		                          {
+			                          return SearchTerms.Contains(Token);
+		                          });
 	}
 
 	bDoesRowMatchSearchTerms = bMatchFound;
@@ -192,6 +213,17 @@ void FConsoleVariablesEditorListRow::SetDoesRowPassFilters(const bool bPass)
 	bDoesRowPassFilters = bPass;
 }
 
+bool FConsoleVariablesEditorListRow::DoesCurrentValueDifferFromPresetValue() const
+{
+	return bDoesCurrentValueDifferFromPresetValue;
+}
+
+void FConsoleVariablesEditorListRow::SetDoesCurrentValueDifferFromPresetValue(
+	const bool bNewValue)
+{
+	bDoesCurrentValueDifferFromPresetValue = bNewValue;
+}
+
 bool FConsoleVariablesEditorListRow::GetIsSelected() const
 {
 	return bIsSelected;
@@ -207,7 +239,8 @@ ECheckBoxState FConsoleVariablesEditorListRow::GetWidgetCheckedState() const
 	return WidgetCheckedState;
 }
 
-void FConsoleVariablesEditorListRow::SetWidgetCheckedState(const ECheckBoxState NewState, const bool bShouldUpdateHierarchyCheckedStates)
+void FConsoleVariablesEditorListRow::SetWidgetCheckedState(const ECheckBoxState NewState,
+                                                           const bool bShouldUpdateHierarchyCheckedStates)
 {
 	WidgetCheckedState = NewState;
 }
@@ -217,10 +250,14 @@ bool FConsoleVariablesEditorListRow::IsRowChecked() const
 	return GetWidgetCheckedState() == ECheckBoxState::Checked;
 }
 
+bool FConsoleVariablesEditorListRow::ShouldBeVisible() const
+{
+	return (bDoesRowMatchSearchTerms && bDoesRowPassFilters) || HasVisibleChildren();
+}
+
 EVisibility FConsoleVariablesEditorListRow::GetDesiredVisibility() const
 {
-	return (bDoesRowMatchSearchTerms && bDoesRowPassFilters) || HasVisibleChildren() ?
-		EVisibility::Visible : EVisibility::Collapsed;
+	return ShouldBeVisible() ? EVisibility::Visible : EVisibility::Collapsed;
 }
 
 const FString& FConsoleVariablesEditorListRow::GetCachedValue() const
@@ -238,23 +275,49 @@ TArray<FConsoleVariablesEditorListRowPtr> FConsoleVariablesEditorListRow::GetSel
 	return ListViewPtr.Pin()->GetSelectedTreeViewItems();
 }
 
-FReply FConsoleVariablesEditorListRow::OnRemoveButtonClicked()
+FReply FConsoleVariablesEditorListRow::OnActionButtonClicked()
 {
-	if (!ensure(ListViewPtr.IsValid()))
+	const bool bIsListValid =
+		GetListViewPtr().IsValid() && GetListViewPtr().Pin()->GetListModelPtr().IsValid();
+
+	if (!ensure(bIsListValid))
 	{
-		return FReply::Handled();
+		return FReply::Unhandled();
 	}
 
-	ResetToStartupValueAndSource();
+	const FConsoleVariablesEditorList::EConsoleVariablesEditorListMode ListMode =
+		bIsListValid
+			? GetListViewPtr().Pin()->GetListModelPtr().Pin()->GetListMode()
+			: FConsoleVariablesEditorList::EConsoleVariablesEditorListMode::Preset;
+
+	const bool bIsGlobalSearch =
+		ListMode == FConsoleVariablesEditorList::EConsoleVariablesEditorListMode::GlobalSearch;
 
 	const FConsoleVariablesEditorModule& ConsoleVariablesEditorModule = FConsoleVariablesEditorModule::Get();
 
-	const TObjectPtr<UConsoleVariablesAsset> EditableAsset = ConsoleVariablesEditorModule.GetEditingAsset();
+	const FString& CommandName = GetCommandInfo().Pin()->Command;
+
+	const TObjectPtr<UConsoleVariablesAsset> EditableAsset = ConsoleVariablesEditorModule.GetPresetAsset();
 	check(EditableAsset);
 
-	EditableAsset->RemoveConsoleVariable(CommandInfo.Pin()->Command);
+	if (bIsGlobalSearch)
+	{
+		EditableAsset->AddOrSetConsoleObjectSavedData(
+			{
+				CommandName,
+				"",
+				ECheckBoxState::Checked
+			}
+		);
+	}
+	else
+	{
+		ResetToStartupValueAndSource();
 
-	ListViewPtr.Pin()->RebuildList();
+		EditableAsset->RemoveConsoleVariable(CommandName);
+
+		ListViewPtr.Pin()->RebuildList();
+	}
 
 	return FReply::Handled();
 }
@@ -265,23 +328,3 @@ void FConsoleVariablesEditorListRow::ResetToPresetValue()
 	GetCommandInfo().Pin()->ExecuteCommand(Value);
 	SetCachedValue(Value);
 }
-
-bool FConsoleVariablesEditorListRow::GetShouldFlashOnScrollIntoView() const
-{
-	return bShouldFlashOnScrollIntoView;
-}
-
-void FConsoleVariablesEditorListRow::SetShouldFlashOnScrollIntoView(const bool bNewShouldFlashOnScrollIntoView)
-{
-	bShouldFlashOnScrollIntoView = bNewShouldFlashOnScrollIntoView;
-}
-
-const FString& FConsoleVariablesEditorListRow::GetPresetValue() const
-{
-	return PresetValue;
-}
-
-void FConsoleVariablesEditorListRow::SetPresetValue(const FString& InPresetValue)
-{
-	PresetValue = InPresetValue;
-}
diff --git a/Engine/Plugins/Experimental/ConsoleVariables/Source/ConsoleVariablesEditor/Private/Views/List/ConsoleVariablesEditorListRow.h b/Engine/Plugins/Experimental/ConsoleVariables/Source/ConsoleVariablesEditor/Private/Views/List/ConsoleVariablesEditorListRow.h
--- a/Engine/Plugins/Experimental/ConsoleVariables/Source/ConsoleVariablesEditor/Private/Views/List/ConsoleVariablesEditorListRow.h
+++ b/Engine/Plugins/Experimental/ConsoleVariables/Source/ConsoleVariablesEditor/Private/Views/List/ConsoleVariablesEditorListRow.h
@@ -37,7 +37,17 @@ struct FConsoleVariablesEditorListRow final : TSharedFromThis<FConsoleVariablesE
 	, ListViewPtr(InListView)
 	, SortOrder(IndexInList)
 	, DirectParentRow(InDirectParentRow)
-	{}
+	{
+		if (PresetValue.IsEmpty())
+		{
+			if (const auto Variable = CommandInfo.Pin()->GetConsoleVariablePtr())
+			{
+				PresetValue = Variable->GetString();
+			}
+		}
+
+		SetCachedValue(PresetValue);
+	}
 
 	[[nodiscard]] TWeakPtr<FConsoleVariablesEditorCommandInfo> GetCommandInfo() const;
 
@@ -83,6 +93,9 @@ struct FConsoleVariablesEditorListRow final : TSharedFromThis<FConsoleVariablesE
 
 	[[nodiscard]] bool GetDoesRowPassFilters() const;
 	void SetDoesRowPassFilters(const bool bPass);
+	
+	[[nodiscard]] bool DoesCurrentValueDifferFromPresetValue() const;
+	void SetDoesCurrentValueDifferFromPresetValue(const bool bNewValue);
 
 	[[nodiscard]] bool GetIsSelected() const;
 	void SetIsSelected(const bool bNewSelected);
@@ -92,6 +105,7 @@ struct FConsoleVariablesEditorListRow final : TSharedFromThis<FConsoleVariablesE
 
 	[[nodiscard]] bool IsRowChecked() const;
 
+	[[nodiscard]] bool ShouldBeVisible() const;
 	[[nodiscard]] EVisibility GetDesiredVisibility() const;
 
 	[[nodiscard]] bool HasVisibleChildren() const
@@ -109,7 +123,7 @@ struct FConsoleVariablesEditorListRow final : TSharedFromThis<FConsoleVariablesE
 
 	[[nodiscard]] TArray<FConsoleVariablesEditorListRowPtr> GetSelectedTreeViewItems() const;
 
-	FReply OnRemoveButtonClicked();
+	FReply OnActionButtonClicked();
 	
 	void ResetToPresetValue();
 
@@ -135,6 +149,8 @@ private:
 	bool bDoesRowMatchSearchTerms = true;
 	bool bDoesRowPassFilters = true;
 
+	bool bDoesCurrentValueDifferFromPresetValue = false;
+	
 	bool bIsSelected = false;
 	TWeakPtr<FConsoleVariablesEditorListRow> DirectParentRow;
 
diff --git a/Engine/Plugins/Experimental/ConsoleVariables/Source/ConsoleVariablesEditor/Private/Views/List/SConsoleVariablesEditorList.cpp b/Engine/Plugins/Experimental/ConsoleVariables/Source/ConsoleVariablesEditor/Private/Views/List/SConsoleVariablesEditorList.cpp
--- a/Engine/Plugins/Experimental/ConsoleVariables/Source/ConsoleVariablesEditor/Private/Views/List/SConsoleVariablesEditorList.cpp
+++ b/Engine/Plugins/Experimental/ConsoleVariables/Source/ConsoleVariablesEditor/Private/Views/List/SConsoleVariablesEditorList.cpp
@@ -3,17 +3,24 @@
 #include "SConsoleVariablesEditorList.h"
 
 #include "ConsoleVariablesAsset.h"
+#include "ConsoleVariablesEditorList.h"
+#include "ConsoleVariablesEditorListFilters/ConsoleVariablesEditorListFilter_ModifiedVariables.h"
 #include "ConsoleVariablesEditorListFilters/ConsoleVariablesEditorListFilter_SourceText.h"
 #include "ConsoleVariablesEditorModule.h"
 #include "ConsoleVariablesEditorProjectSettings.h"
 #include "SConsoleVariablesEditorListRow.h"
+#include "../Widgets/SConsoleVariablesEditorGlobalSearchToggle.h"
 
+#include "Algo/AllOf.h"
+#include "Algo/AnyOf.h"
 #include "Algo/Find.h"
 #include "Framework/MultiBox/MultiBoxBuilder.h"
+#include "Styling/StyleColors.h"
 #include "Widgets/Input/SCheckBox.h"
 #include "Widgets/Input/SComboBox.h"
 #include "Widgets/Input/SSearchBox.h"
 #include "Widgets/Layout/SWidgetSwitcher.h"
+#include "Widgets/Layout/SWrapBox.h"
 
 #define LOCTEXT_NAMESPACE "ConsoleVariablesEditor"
 
@@ -23,8 +30,10 @@
 const FName SConsoleVariablesEditorList::ValueColumnName(TEXT("Value"));
 const FName SConsoleVariablesEditorList::SourceColumnName(TEXT("Source"));
 
-void SConsoleVariablesEditorList::Construct(const FArguments& InArgs)
+void SConsoleVariablesEditorList::Construct(const FArguments& InArgs, TSharedRef<FConsoleVariablesEditorList> ListModel)
 {
+	ListModelPtr = ListModel;
+	
 	// Set Default Sorting info
 	ActiveSortingColumnName = CustomSortOrderColumnName;
 	ActiveSortingType = EColumnSortMode::Ascending;
@@ -55,5 +64,26 @@
 			]
 
 			+SHorizontalBox::Slot()
+			.Padding(4.f, 0.f, 2.f, 0.f)
+			.VAlign(VAlign_Center)
 			.AutoWidth()
+			[
+				SNew(SButton)
+				.VAlign(VAlign_Center)
+				.OnClicked_Lambda([this]()
+				{
+					const FString SearchString = GetSearchStringFromSearchInputField().ToLower().TrimStartAndEnd();
+					ListSearchBoxPtr->SetText(FText::GetEmpty());
+					return TryEnterGlobalSearch(SearchString);
+				})
+				.ToolTipText(LOCTEXT("OpenInGlobalSearchButtonTooltip", "Search All Console Variables"))
+				[
+					SNew(STextBlock)
+					.TextStyle(FAppStyle::Get(), "FindResults.FindInBlueprints")
+					.Text(FText::FromString(FString(TEXT("\xf1e5"))) /*fa-binoculars*/)
+				]
+			]
+
+			+SHorizontalBox::Slot()
+			.AutoWidth()
 			.Padding(10.f, 1.f, 15.f, 1.f)
@@ -72,6 +102,57 @@
 			]
 		]
 
+		+SVerticalBox::Slot()
+		.AutoHeight()
+		[
+			SAssignNew(GlobalSearchesHBox, SHorizontalBox)
+			.Visibility_Lambda([this]()
+			{
+				const bool bShouldBeVisible =
+					GlobalSearchesContainer.IsValid() && GlobalSearchesContainer->GetChildren()->Num() > 0;
+
+				return bShouldBeVisible ? EVisibility::SelfHitTestInvisible : EVisibility::Collapsed;
+			})
+
+			+SHorizontalBox::Slot()
+			.HAlign(HAlign_Fill)
+			[
+				SAssignNew(GlobalSearchesContainer, SWrapBox)
+				.UseAllottedSize(true)
+			]
+
+			+SHorizontalBox::Slot()
+			.Padding(4.f, 0.f, 2.f, 0.f)
+			.VAlign(VAlign_Center)
+			.HAlign(HAlign_Right)
+			.AutoWidth()
+			[
+				// Remove Button
+				SAssignNew(RemoveGlobalSearchesButtonPtr, SButton)
+				.ToolTipText(
+					LOCTEXT("RemoveGlobalSearchesButtonTooltip",
+					"Remove all global searches from the console variables editor."))
+				.ButtonColorAndOpacity(FStyleColors::Transparent)
+				.OnClicked_Lambda([this]()
+				{
+					ListModelPtr.Pin()->SetListMode(FConsoleVariablesEditorList::EConsoleVariablesEditorListMode::Preset);
+
+					GlobalSearchesContainer->ClearChildren();
+					CurrentGlobalSearches.Empty();
+
+					RebuildList();
+					
+					return FReply::Handled();
+				})
+				[
+					SNew(SImage)
+					.Visibility(EVisibility::SelfHitTestInvisible)
+					.Image(FAppStyle::Get().GetBrush("Icons.Delete"))
+					.ColorAndOpacity(FSlateColor::UseForeground())
+				]
+			]
+		]
+
 		+ SVerticalBox::Slot()
 		[
 			SNew(SWidgetSwitcher)
@@ -94,7 +175,7 @@
 						Row->SetIsSelected(TreeViewPtr->GetSelectedItems().Contains(Row));
 					}
 				})
-				.TreeItemsSource(&TreeViewRootObjects)
+				.TreeItemsSource(&VisibleTreeViewObjects)
 				.OnGenerateRow_Lambda([this](FConsoleVariablesEditorListRowPtr Row, const TSharedRef<STableViewBase>& OwnerTable)
 					{
 						check(Row.IsValid());
@@ -125,150 +206,148 @@
 
 SConsoleVariablesEditorList::~SConsoleVariablesEditorList()
 {	
+	HeaderRow.Reset();
+	
 	ListSearchBoxPtr.Reset();
-	ListBoxContainerPtr.Reset();
 	ViewOptionsComboButton.Reset();
+	GlobalSearchesHBox.Reset();
+	GlobalSearchesContainer.Reset();
+	CurrentGlobalSearches.Empty();
+	RemoveGlobalSearchesButtonPtr.Reset();
+	ListBoxContainerPtr.Reset();
 
 	FlushMemory(false);
 
-	HeaderRow.Reset();
-
-	TreeViewRootObjects.Empty();
+	ShowFilters.Reset();
 	TreeViewPtr.Reset();
+	VisibleTreeViewObjects.Reset();
+	LastPresetObjects.Reset();
 }
 
-void SConsoleVariablesEditorList::SetupFilters()
+FReply SConsoleVariablesEditorList::TryEnterGlobalSearch(const FString& SearchString)
 {
-	TArray<FString> SourceFilterTypes =
+	FReply ReturnValue = FReply::Unhandled();
+
+	// Can't enter global search if there are no active global searches or a search string from which to parse new searches
+	if (SearchString.IsEmpty() && CurrentGlobalSearches.Num() == 0)
 	{
-		"Constructor",
-		"Scalability",
-		"Game Setting",
-		"Project Setting",
-		"System Settings ini",
-		"Device Profile",
-		"Game Override",
-		"Console Variables ini",
-		"Command line",
-		"Code",
-		"Console"
-	};
+		UE_LOG(LogConsoleVariablesEditor, Log,
+			TEXT("%hs: Global search request is empty. Exiting Global Search."),
+			__FUNCTION__, *SearchString);
 
-	for (const FString& Type : SourceFilterTypes)
+		// Return to preset mode if in global search mode then rebuild list
+		if (ListModelPtr.Pin()->GetListMode() == FConsoleVariablesEditorList::EConsoleVariablesEditorListMode::GlobalSearch)
+		{
+			ListModelPtr.Pin()->SetListMode(FConsoleVariablesEditorList::EConsoleVariablesEditorListMode::Preset);
+			RebuildList();
+		}
+
+		return ReturnValue;
+	}
+
+	// Strings that already have associated buttons
+	TArray<FString> ExistingButtonStrings;
+	// All strings parsed from the search text
+	TArray<FString> OutTokens;
+	SearchString.ParseIntoArray(OutTokens, TEXT(" "));
+
+	// Get tokens from current searches. This step allows us to properly populate the asset with all matching commands
+	for (const TSharedRef<SConsoleVariablesEditorGlobalSearchToggle>& GlobalSearchButton : CurrentGlobalSearches)
 	{
-		ShowFilters.Add(MakeShared<FConsoleVariablesEditorListFilter_SourceText>(Type));
-	}
-}
+		const FString ButtonText = GlobalSearchButton->GetGlobalSearchText().ToString();
+		ExistingButtonStrings.Add(ButtonText);
 
-TSharedRef<SWidget> SConsoleVariablesEditorList::BuildShowOptionsMenu()
-{
-	FMenuBuilder ShowOptionsMenuBuilder = FMenuBuilder(true, nullptr);
+		// If the button text was explicitly typed into the search then set the existing button to be checked.
+		if (OutTokens.ContainsByPredicate([this, &ButtonText](const FString& Token)
+		{
+			return Token.Equals(ButtonText, ESearchCase::IgnoreCase);
+		}))
+		{
+			GlobalSearchButton->SetIsButtonChecked(true);
+		}
 
-	ShowOptionsMenuBuilder.BeginSection("", LOCTEXT("ShowOptions_ShowSectionHeading", "Show"));
-	{
-		// Add mode filters
-		auto AddFiltersLambda = [this, &ShowOptionsMenuBuilder](const TSharedRef<IConsoleVariablesEditorListFilter>& InFilter)
+		// If the button is checked, add its search text to the tokens array so it can be used to populate the matching command list
+		if (GlobalSearchButton->GetIsToggleChecked())
 		{
-			const FString& FilterName = InFilter->GetFilterName();
-			
-			ShowOptionsMenuBuilder.AddMenuEntry(
-			   InFilter->GetFilterButtonLabel(),
-			   InFilter->GetFilterButtonToolTip(),
-			   FSlateIcon(),
-			   FUIAction(
-				   FExecuteAction::CreateLambda(
-				   	[this, FilterName]()
-					   {
-						   ToggleFilterActive(FilterName);
-					   }
-					),
-				   FCanExecuteAction(),
-				   FIsActionChecked::CreateSP( InFilter, &IConsoleVariablesEditorListFilter::GetIsFilterActive )
-			   ),
-			   NAME_None,
-			   EUserInterfaceActionType::ToggleButton
-		   );
-		};
-
-		for (const TSharedRef<IConsoleVariablesEditorListFilter>& Filter : ShowFilters)
-		{
-			AddFiltersLambda(Filter);
+			OutTokens.AddUnique(ButtonText);
 		}
 	}
-	ShowOptionsMenuBuilder.EndSection();
 	
-	ShowOptionsMenuBuilder.BeginSection("", LOCTEXT("ShowOptions_SortSectionHeading", "Sort"));
+	FConsoleVariablesEditorModule& ConsoleVariablesEditorModule = FConsoleVariablesEditorModule::Get();
+	
+	if (ConsoleVariablesEditorModule.PopulateGlobalSearchAssetWithVariablesMatchingTokens(OutTokens))
 	{
-		// Add commands
+		// If we were in Preset mode before entering global search, cache the existing tree objects to maintain state
+		if (ListModelPtr.Pin()->GetListMode() == FConsoleVariablesEditorList::EConsoleVariablesEditorListMode::Preset)
+		{
+			LastPresetObjects = TreeViewRootObjects;
+			ListModelPtr.Pin()->SetListMode(FConsoleVariablesEditorList::EConsoleVariablesEditorListMode::GlobalSearch);
+		}
 
-		// Save this for later when folders are added
-		/*ShowOptionsMenuBuilder.AddMenuEntry(
-			LOCTEXT("CollapseAll", "Collapse All"),
-			LOCTEXT("ConsoleVariablesEditorList_CollapseAll_Tooltip", "Collapse all expanded actor groups in the Modified Actors list."),
-			FSlateIcon(),
-			FUIAction(FExecuteAction::CreateRaw(this, &SConsoleVariablesEditorList::SetAllGroupsCollapsed)),
-			NAME_None,
-			EUserInterfaceActionType::Button
-		);*/
-		
-		ShowOptionsMenuBuilder.AddMenuEntry(
-			LOCTEXT("SetSortOrder", "Set Sort Order"),
-			LOCTEXT("ConsoleVariablesEditorList_SetSortOrder_Tooltip", "Makes the current order of the variables list the saved order."),
-			FSlateIcon(),
-			FUIAction(FExecuteAction::CreateRaw(this, &SConsoleVariablesEditorList::SetSortOrder)),
-			NAME_None,
-			EUserInterfaceActionType::Button
-		);
-	}
-	ShowOptionsMenuBuilder.EndSection();
-	
-	ShowOptionsMenuBuilder.BeginSection("", LOCTEXT("ShowOptions_OptionsSectionHeading", "Options"));
-	{
-		ShowOptionsMenuBuilder.AddMenuEntry(
-			LOCTEXT("TrackAllVariableChanges", "Track All Variable Changes"),
-			LOCTEXT("ConsoleVariablesEditorList_TrackAllVariableChanges_Tooltip",
-				"When variables are changed outside the Console Variables Editor, this option will add the variables to the current preset. Does not apply to console commands like 'r.SetNearClipPlane' or 'stat fps'."),
-			FSlateIcon(),
-			FUIAction(
-				FExecuteAction::CreateLambda([]()
+		// Convert tokens to global search toggle buttons
+		for (const FString& TokenString : OutTokens)
+		{
+			// Only make new buttons when one doesn't exist for the current token
+			if (!ExistingButtonStrings.Contains(TokenString))
+			{
+				TSharedRef<SConsoleVariablesEditorGlobalSearchToggle> NewGlobalSearchButton =
+					SNew(SConsoleVariablesEditorGlobalSearchToggle, FText::FromString(TokenString))
+				.OnToggleClickedOnce_Lambda([this]()
 				{
-					if (UConsoleVariablesEditorProjectSettings* ProjectSettingsPtr =
-						GetMutableDefault<UConsoleVariablesEditorProjectSettings>())
-					{
-						ProjectSettingsPtr->bAddAllChangedConsoleVariablesToCurrentPreset =
-							!ProjectSettingsPtr->bAddAllChangedConsoleVariablesToCurrentPreset;
-					}
-				}),
-			   FCanExecuteAction(),
-			   FIsActionChecked::CreateLambda( []()
-				{
-					if (const UConsoleVariablesEditorProjectSettings* ProjectSettingsPtr =
-						GetMutableDefault<UConsoleVariablesEditorProjectSettings>())
-					{
-						return ProjectSettingsPtr->bAddAllChangedConsoleVariablesToCurrentPreset;
-					}
-
-			   		return false;
+					return TryEnterGlobalSearch();
 				})
-			),
-			NAME_None,
-			EUserInterfaceActionType::ToggleButton
-		);
+				.OnToggleCtrlClicked(this, &SConsoleVariablesEditorList::HandleRemoveGlobalSearchToggleButton)
+				.OnToggleRightButtonClicked(this, &SConsoleVariablesEditorList::HandleRemoveGlobalSearchToggleButton);
+
+				CurrentGlobalSearches.Add(NewGlobalSearchButton);
+			}
+		}
+
+		// Put widgets in container
+		RefreshGlobalSearchWidgets();
+
+		ReturnValue = FReply::Handled();
 	}
-	ShowOptionsMenuBuilder.EndSection();
-
-	return ShowOptionsMenuBuilder.MakeWidget();
-}
-
-void SConsoleVariablesEditorList::FlushMemory(const bool bShouldKeepMemoryAllocated)
-{
-	if (bShouldKeepMemoryAllocated)
-	{
-		TreeViewRootObjects.Reset();
-	}
 	else
 	{
-		TreeViewRootObjects.Empty();
+		UE_LOG(LogConsoleVariablesEditor, Warning,
+			TEXT("%hs: Failed to find console variable objects with names containing search string %s"),
+			__FUNCTION__, *SearchString);
+	}
+
+	RebuildList();
+
+	return ReturnValue;
+}
+
+FReply SConsoleVariablesEditorList::HandleRemoveGlobalSearchToggleButton()
+{
+	CleanUpGlobalSearchesMarkedForDelete();
+	RefreshGlobalSearchWidgets();
+	return TryEnterGlobalSearch();
+}
+
+void SConsoleVariablesEditorList::CleanUpGlobalSearchesMarkedForDelete()
+{
+	for (int32 GlobalSearchItr = CurrentGlobalSearches.Num() - 1; GlobalSearchItr >= 0; GlobalSearchItr--)
+	{
+		if (CurrentGlobalSearches[GlobalSearchItr]->GetIsMarkedForDelete())
+		{
+			CurrentGlobalSearches.RemoveAt(GlobalSearchItr);
+		}
 	}
 }
 
+void SConsoleVariablesEditorList::RefreshGlobalSearchWidgets()
+{
+	GlobalSearchesContainer->ClearChildren();
+
+	for (const TSharedRef<SConsoleVariablesEditorGlobalSearchToggle>& GlobalSearchButton : CurrentGlobalSearches)
+	{	
+		GlobalSearchesContainer->AddSlot()
+		[
+			GlobalSearchButton
+		];
+	}
+}
+
@@ -304,19 +383,19 @@
 	if (TreeViewRootObjects.Num() > 0)
 	{
 		// Apply last search
-		ExecuteListViewSearchOnAllRows(GetSearchStringFromSearchInputField());
+		ExecuteListViewSearchOnAllRows(GetSearchStringFromSearchInputField(), false);
 
 		// Enforce Sort
 		const FName& SortingName = GetActiveSortingColumnName();
-		ExecuteSort(SortingName, GetSortModeForColumn(SortingName));
+		ExecuteSort(SortingName, GetSortModeForColumn(SortingName), false);
 
 		// Show/Hide rows based on SetBy changes and filter settings
-		EvaluateIfRowsPassFilters();
+		EvaluateIfRowsPassFilters(false);
 
 		// Refresh the header's check state
 		OnListItemCheckBoxStateChange(ECheckBoxState::Undetermined);
 
-		TreeViewPtr->RequestTreeRefresh();
+		FindVisibleObjectsAndRequestTreeRefresh();
 	}
 }
 
@@ -362,66 +441,246 @@
 
 FString SConsoleVariablesEditorList::GetSearchStringFromSearchInputField() const
 {
-	return ensureAlwaysMsgf(ListSearchBoxPtr.IsValid(), TEXT("%hs: ListSearchBoxPtr is not valid. Check to make sure it was created."), __FUNCTION__)
+	return ensureAlwaysMsgf(ListSearchBoxPtr.IsValid(),
+		TEXT("%hs: ListSearchBoxPtr is not valid. Check to make sure it was created."), __FUNCTION__)
 	? ListSearchBoxPtr->GetText().ToString() : "";
 }
 
-void SConsoleVariablesEditorList::GenerateTreeView()
-{	
-	if (!ensure(TreeViewPtr.IsValid()))
+void SConsoleVariablesEditorList::ExecuteListViewSearchOnAllRows(
+	const FString& SearchString, const bool bShouldRefreshAfterward)
+{
+	TArray<FString> Tokens;
+	
+	// unquoted search equivalent to a match-any-of search
+	SearchString.ParseIntoArray(Tokens, TEXT(" "), true);
+	
+	for (const TSharedPtr<FConsoleVariablesEditorListRow>& ChildRow : TreeViewRootObjects)
 	{
-		return;
-	}
-	
-	FlushMemory(true);
-
-	FConsoleVariablesEditorModule& ConsoleVariablesEditorModule = FConsoleVariablesEditorModule::Get();
-
-	const TObjectPtr<UConsoleVariablesAsset> EditableAsset = ConsoleVariablesEditorModule.GetEditingAsset();
-	check(EditableAsset);
-
-	for (const FConsoleVariablesEditorAssetSaveData& SavedCommand : EditableAsset->GetSavedCommands())
-	{
-		// Get corresponding CommandInfo for tracking or make one if the command is non-value
-		TSharedPtr<FConsoleVariablesEditorCommandInfo> CommandInfo =
-			ConsoleVariablesEditorModule.FindCommandInfoByName(SavedCommand.CommandName).Pin();
-
-		if (!CommandInfo.IsValid())
+		if (!ensure(ChildRow.IsValid()))
 		{
-			CommandInfo = MakeShared<FConsoleVariablesEditorCommandInfo>(SavedCommand.CommandName);
-
-			FConsoleVariablesEditorModule::Get().AddConsoleObjectCommandInfoToMasterReference(
-				CommandInfo.ToSharedRef());
+			continue;
 		}
 		
-		if (CommandInfo.IsValid())
+		const bool bGroupMatch = ChildRow->MatchSearchTokensToSearchTerms(Tokens);
+		
+		// If the group name matches then we pass in an empty string so all child nodes are visible.
+		// If the name doesn't match, then we need to evaluate each child.
+		ChildRow->ExecuteSearchOnChildNodes(bGroupMatch ? "" : SearchString);
+	}
+
+	if (bShouldRefreshAfterward)
+	{
+		FindVisibleObjectsAndRequestTreeRefresh();
+	}
+}
+
+bool SConsoleVariablesEditorList::DoesTreeViewHaveVisibleChildren() const
+{
+	if (TreeViewPtr.IsValid())
+	{
+		for (const TSharedPtr<FConsoleVariablesEditorListRow>& Header : TreeViewRootObjects)
 		{
-			const ECheckBoxState NewCheckedState =
-				SavedCommand.CheckedState == ECheckBoxState::Unchecked ?
-						ECheckBoxState::Unchecked : ECheckBoxState::Checked;
+			const EVisibility HeaderVisibility = Header->GetDesiredVisibility();
 			
-			// If the row is checked and the saved value differs from the current value,
-			// execute the command with the saved value 
-			if (CommandInfo->ObjectType == FConsoleVariablesEditorCommandInfo::EConsoleObjectType::Command ||
-				(NewCheckedState == ECheckBoxState::Checked &&
-				CommandInfo->IsCurrentValueDifferentFromInputValue(SavedCommand.CommandValueAsString)))
+			if (HeaderVisibility != EVisibility::Hidden && HeaderVisibility != EVisibility::Collapsed)
 			{
-				CommandInfo->ExecuteCommand(SavedCommand.CommandValueAsString);
+				return true;
 			}
-		
-			FConsoleVariablesEditorListRowPtr NewRow = 
-				MakeShared<FConsoleVariablesEditorListRow>(
-						CommandInfo, SavedCommand.CommandValueAsString,
-						FConsoleVariablesEditorListRow::SingleCommand, 
-						NewCheckedState, SharedThis(this), TreeViewRootObjects.Num(), nullptr);
-		
-			TreeViewRootObjects.Add(NewRow);
+		}
+	}
+	
+	return false;
+}
+
+void SConsoleVariablesEditorList::SetTreeViewItemExpanded(const TSharedPtr<FConsoleVariablesEditorListRow>& RowToExpand, const bool bNewExpansion) const
+{
+	if (TreeViewPtr.IsValid())
+	{
+		TreeViewPtr->SetItemExpansion(RowToExpand, bNewExpansion);
+	}
+}
+
+void SConsoleVariablesEditorList::SetAllListViewItemsCheckState(const ECheckBoxState InNewState)
+{
+	for (const TSharedPtr<FConsoleVariablesEditorListRow>& Row : TreeViewRootObjects)
+	{
+		Row->SetWidgetCheckedState(InNewState);
+	}
+}
+
+bool SConsoleVariablesEditorList::DoesListHaveCheckedMembers() const
+{
+	for (const TSharedPtr<FConsoleVariablesEditorListRow>& Row : TreeViewRootObjects)
+	{
+		if (Row->GetCommandInfo().Pin()->ObjectType ==
+			FConsoleVariablesEditorCommandInfo::EConsoleObjectType::Variable &&
+			Row->GetWidgetCheckedState() == ECheckBoxState::Checked)
+		{
+			return true;
 		}
 	}
 
-	TreeViewPtr->RequestTreeRefresh();
+	return false;
 }
 
+bool SConsoleVariablesEditorList::DoesListHaveUncheckedMembers() const
+{
+	for (const TSharedPtr<FConsoleVariablesEditorListRow>& Row : TreeViewRootObjects)
+	{
+		if (Row->GetCommandInfo().Pin()->ObjectType ==
+			FConsoleVariablesEditorCommandInfo::EConsoleObjectType::Variable &&
+			Row->GetWidgetCheckedState() == ECheckBoxState::Unchecked)
+		{
+			return true;
+		}
+	}
+
+	return false;
+}
+
+void SConsoleVariablesEditorList::OnListItemCheckBoxStateChange(const ECheckBoxState InNewState)
+{
+	HeaderCheckBoxState = ECheckBoxState::Checked;
+
+	if (DoesListHaveUncheckedMembers())
+	{
+		HeaderCheckBoxState = ECheckBoxState::Unchecked;
+
+		if (DoesListHaveCheckedMembers())
+		{
+			HeaderCheckBoxState = ECheckBoxState::Undetermined;
+		}
+	}
+}
+
+void SConsoleVariablesEditorList::ToggleFilterActive(const FString& FilterName)
+{
+	if (const TSharedRef<IConsoleVariablesEditorListFilter>* Match =
+		Algo::FindByPredicate(ShowFilters,
+		[&FilterName](TSharedRef<IConsoleVariablesEditorListFilter> Comparator)
+		{
+			return Comparator->GetFilterName().Equals(FilterName);
+		}))
+	{
+		const TSharedRef<IConsoleVariablesEditorListFilter> Filter = *Match;
+		Filter->ToggleFilterActive();
+
+		EvaluateIfRowsPassFilters();
+	}
+}
+
+void SConsoleVariablesEditorList::EvaluateIfRowsPassFilters(const bool bShouldRefreshAfterward)
+{
+	// Separate filters by type
+	
+	TSet<TSharedRef<IConsoleVariablesEditorListFilter>> MatchAnyOfFilters;
+	TSet<TSharedRef<IConsoleVariablesEditorListFilter>> MatchAllOfFilters;
+
+	for (const TSharedRef<IConsoleVariablesEditorListFilter>& Filter : ShowFilters)
+	{
+		if (Filter->GetFilterMatchType() ==
+			IConsoleVariablesEditorListFilter::EConsoleVariablesEditorListFilterMatchType::MatchAll)
+		{
+			MatchAllOfFilters.Add(Filter);
+		}
+		else
+		{
+			MatchAnyOfFilters.Add(Filter);
+		}
+	}
+	
+	for (const FConsoleVariablesEditorListRowPtr& Row : TreeViewRootObjects)
+	{
+		if (Row.IsValid() && Row->GetRowType() == FConsoleVariablesEditorListRow::SingleCommand)
+		{
+			auto Projection = [&Row](const TSharedRef<IConsoleVariablesEditorListFilter>& Filter)
+			{
+				return Filter->GetIsFilterActive() ? Filter->DoesItemPassFilter(Row) : true;
+			};
+			
+			const bool bPassesAnyOf = Algo::AnyOf(MatchAnyOfFilters, Projection);
+			const bool bPassesAllOf = Algo::AllOf(MatchAllOfFilters, Projection);
+			
+			Row->SetDoesRowPassFilters(bPassesAnyOf && bPassesAllOf);
+		}
+	}
+
+	if (bShouldRefreshAfterward)
+	{
+		FindVisibleObjectsAndRequestTreeRefresh();
+	}
+}
+
+EColumnSortMode::Type SConsoleVariablesEditorList::GetSortModeForColumn(FName InColumnName) const
+{
+	EColumnSortMode::Type ColumnSortMode = EColumnSortMode::None;
+
+	if (GetActiveSortingColumnName().IsEqual(InColumnName))
+	{
+		ColumnSortMode = ActiveSortingType;
+	}
+
+	return ColumnSortMode;
+}
+
+void SConsoleVariablesEditorList::OnSortColumnCalled(EColumnSortPriority::Type Priority, const FName& ColumnName, EColumnSortMode::Type SortMode)
+{
+	ExecuteSort(ColumnName, CycleSortMode(ColumnName));
+}
+
+EColumnSortMode::Type SConsoleVariablesEditorList::CycleSortMode(const FName& InColumnName)
+{
+	// Custom handler for Custom Sort Order mode
+	if (InColumnName.IsEqual(CustomSortOrderColumnName))
+	{
+		ActiveSortingType = EColumnSortMode::Ascending;
+	}
+	else
+	{
+		const EColumnSortMode::Type PreviousColumnSortMode = GetSortModeForColumn(InColumnName);
+		ActiveSortingType = PreviousColumnSortMode ==
+			EColumnSortMode::Ascending ? EColumnSortMode::Descending : EColumnSortMode::Ascending;
+	}
+
+	ActiveSortingColumnName = InColumnName;
+	return ActiveSortingType;
+}
+
+void SConsoleVariablesEditorList::ExecuteSort(
+	const FName& InColumnName, const EColumnSortMode::Type InColumnSortMode, const bool bShouldRefreshAfterward)
+{	
+	if (InColumnName.IsEqual(CustomSortOrderColumnName))
+	{
+		TreeViewRootObjects.StableSort(SortByOrderAscending);
+	}
+	if (InColumnName.IsEqual(SourceColumnName))
+	{
+		TreeViewRootObjects.StableSort(
+			InColumnSortMode == EColumnSortMode::Ascending ? SortBySourceAscending : SortBySourceDescending);
+	}
+	if (InColumnName.IsEqual(VariableNameColumnName))
+	{
+		TreeViewRootObjects.StableSort(
+			InColumnSortMode == EColumnSortMode::Ascending ? SortByVariableNameAscending : SortByVariableNameDescending);
+	}
+
+	if (bShouldRefreshAfterward)
+	{
+		FindVisibleObjectsAndRequestTreeRefresh();
+	}
+}
+
+void SConsoleVariablesEditorList::SetSortOrder(const bool bShouldRefreshAfterward)
+{
+	for (int32 RowItr = 0; RowItr < TreeViewRootObjects.Num(); RowItr++)
+	{
+		const TSharedPtr<FConsoleVariablesEditorListRow>& ChildRow = TreeViewRootObjects[RowItr];
+		ChildRow->SetSortOrder(RowItr);
+	}
+
+	ExecuteSort(CustomSortOrderColumnName, CycleSortMode(CustomSortOrderColumnName), bShouldRefreshAfterward);
+}
+
 TSharedPtr<SHeaderRow> SConsoleVariablesEditorList::GenerateHeaderRow()
 {
 	check(HeaderRow);
@@ -497,6 +756,144 @@
 	return HeaderRow;
 }
 
+void SConsoleVariablesEditorList::SetupFilters()
+{
+	TArray<FString> SourceFilterTypes =
+	{
+		"Constructor",
+		"Scalability",
+		"Game Setting",
+		"Project Setting",
+		"System Settings ini",
+		"Device Profile",
+		"Game Override",
+		"Console Variables ini",
+		"Command line",
+		"Code",
+		"Console"
+	};
+
+	for (const FString& Type : SourceFilterTypes)
+	{
+		ShowFilters.Add(MakeShared<FConsoleVariablesEditorListFilter_SourceText>(Type));
+	}
+
+	// Add Show Only Modified filter
+	ShowFilters.Add(MakeShared<ConsoleVariablesEditorListFilter_ModifiedVariables>());
+}
+
+TSharedRef<SWidget> SConsoleVariablesEditorList::BuildShowOptionsMenu()
+{
+	FMenuBuilder ShowOptionsMenuBuilder = FMenuBuilder(true, nullptr);
+
+	ShowOptionsMenuBuilder.BeginSection("", LOCTEXT("ShowOptions_ShowSectionHeading", "Show"));
+	{
+		// Add show filters
+		auto AddFiltersLambda = [this, &ShowOptionsMenuBuilder](const TSharedRef<IConsoleVariablesEditorListFilter>& InFilter)
+		{
+			const FString& FilterName = InFilter->GetFilterName();
+			
+			ShowOptionsMenuBuilder.AddMenuEntry(
+			   InFilter->GetFilterButtonLabel(),
+			   InFilter->GetFilterButtonToolTip(),
+			   FSlateIcon(),
+			   FUIAction(
+				   FExecuteAction::CreateLambda(
+				   	[this, FilterName]()
+					   {
+						   ToggleFilterActive(FilterName);
+					   }
+					),
+				   FCanExecuteAction(),
+				   FIsActionChecked::CreateSP( InFilter, &IConsoleVariablesEditorListFilter::GetIsFilterActive )
+			   ),
+			   NAME_None,
+			   EUserInterfaceActionType::ToggleButton
+		   );
+		};
+
+		for (const TSharedRef<IConsoleVariablesEditorListFilter>& Filter : ShowFilters)
+		{
+			AddFiltersLambda(Filter);
+		}
+	}
+	ShowOptionsMenuBuilder.EndSection();
+	
+	ShowOptionsMenuBuilder.BeginSection("", LOCTEXT("ShowOptions_SortSectionHeading", "Sort"));
+	{
+		// Add commands
+
+		// Save this for later when folders are added
+		/*ShowOptionsMenuBuilder.AddMenuEntry(
+			LOCTEXT("CollapseAll", "Collapse All"),
+			LOCTEXT("ConsoleVariablesEditorList_CollapseAll_Tooltip", "Collapse all expanded actor groups in the Modified Actors list."),
+			FSlateIcon(),
+			FUIAction(FExecuteAction::CreateRaw(this, &SConsoleVariablesEditorList::SetAllGroupsCollapsed)),
+			NAME_None,
+			EUserInterfaceActionType::Button
+		);*/
+		
+		ShowOptionsMenuBuilder.AddMenuEntry(
+			LOCTEXT("SetSortOrder", "Set Sort Order"),
+			LOCTEXT("ConsoleVariablesEditorList_SetSortOrder_Tooltip", "Makes the current order of the variables list the saved order."),
+			FSlateIcon(),
+			FUIAction(FExecuteAction::CreateRaw(this, &SConsoleVariablesEditorList::SetSortOrder, true)),
+			NAME_None,
+			EUserInterfaceActionType::Button
+		);
+	}
+	ShowOptionsMenuBuilder.EndSection();
+	
+	ShowOptionsMenuBuilder.BeginSection("", LOCTEXT("ShowOptions_OptionsSectionHeading", "Options"));
+	{
+		ShowOptionsMenuBuilder.AddMenuEntry(
+			LOCTEXT("TrackAllVariableChanges", "Track All Variable Changes"),
+			LOCTEXT("ConsoleVariablesEditorList_TrackAllVariableChanges_Tooltip",
+				"When variables are changed outside the Console Variables Editor, this option will add the variables to the current preset. Does not apply to console commands like 'r.SetNearClipPlane' or 'stat fps'."),
+			FSlateIcon(),
+			FUIAction(
+				FExecuteAction::CreateLambda([]()
+				{
+					if (UConsoleVariablesEditorProjectSettings* ProjectSettingsPtr =
+						GetMutableDefault<UConsoleVariablesEditorProjectSettings>())
+					{
+						ProjectSettingsPtr->bAddAllChangedConsoleVariablesToCurrentPreset =
+							!ProjectSettingsPtr->bAddAllChangedConsoleVariablesToCurrentPreset;
+					}
+				}),
+			   FCanExecuteAction(),
+			   FIsActionChecked::CreateLambda( []()
+				{
+					if (const UConsoleVariablesEditorProjectSettings* ProjectSettingsPtr =
+						GetMutableDefault<UConsoleVariablesEditorProjectSettings>())
+					{
+						return ProjectSettingsPtr->bAddAllChangedConsoleVariablesToCurrentPreset;
+					}
+
+			   		return false;
+				})
+			),
+			NAME_None,
+			EUserInterfaceActionType::ToggleButton
+		);
+	}
+	ShowOptionsMenuBuilder.EndSection();
+
+	return ShowOptionsMenuBuilder.MakeWidget();
+}
+
+void SConsoleVariablesEditorList::FlushMemory(const bool bShouldKeepMemoryAllocated)
+{
+	if (bShouldKeepMemoryAllocated)
+	{
+		TreeViewRootObjects.Reset();
+	}
+	else
+	{
+		TreeViewRootObjects.Empty();
+	}
+}
+
 void SConsoleVariablesEditorList::SetAllGroupsCollapsed()
 {
 	if (TreeViewPtr.IsValid())
@@ -514,211 +911,124 @@
 	}
 }
 
-void SConsoleVariablesEditorList::SetSortOrder()
+void SConsoleVariablesEditorList::OnListViewSearchTextChanged(const FText& Text)
 {
-	for (int32 RowItr = 0; RowItr < TreeViewRootObjects.Num(); RowItr++)
-	{
-		const TSharedPtr<FConsoleVariablesEditorListRow>& ChildRow = TreeViewRootObjects[RowItr];
-		ChildRow->SetSortOrder(RowItr);
-	}
-
-	ExecuteSort(CustomSortOrderColumnName, CycleSortMode(CustomSortOrderColumnName));
-}
-
-void SConsoleVariablesEditorList::OnListViewSearchTextChanged(const FText& Text) const
-{
 	ExecuteListViewSearchOnAllRows(Text.ToString());
 }
 
-void SConsoleVariablesEditorList::ExecuteListViewSearchOnAllRows(const FString& SearchString) const
-{
-	TArray<FString> Tokens;
-	
-	// unquoted search equivalent to a match-any-of search
-	SearchString.ParseIntoArray(Tokens, TEXT(" "), true);
-	
-	for (const TSharedPtr<FConsoleVariablesEditorListRow>& ChildRow : TreeViewRootObjects)
-	{
-		if (!ensure(ChildRow.IsValid()))
-		{
-			continue;
-		}
-		
-		const bool bGroupMatch = ChildRow->MatchSearchTokensToSearchTerms(Tokens);
-		
-		// If the group name matches then we pass in an empty string so all child nodes are visible.
-		// If the name doesn't match, then we need to evaluate each child.
-		ChildRow->ExecuteSearchOnChildNodes(bGroupMatch ? "" : SearchString);
-	}
-
-	TreeViewPtr->RequestTreeRefresh();
-}
-
-bool SConsoleVariablesEditorList::DoesTreeViewHaveVisibleChildren() const
-{
-	if (TreeViewPtr.IsValid())
-	{
-		for (const TSharedPtr<FConsoleVariablesEditorListRow>& Header : TreeViewRootObjects)
-		{
-			const EVisibility HeaderVisibility = Header->GetDesiredVisibility();
-			
-			if (HeaderVisibility != EVisibility::Hidden && HeaderVisibility != EVisibility::Collapsed)
-			{
-				return true;
-			}
-		}
-	}
-	
-	return false;
-}
-
-void SConsoleVariablesEditorList::SetTreeViewItemExpanded(const TSharedPtr<FConsoleVariablesEditorListRow>& RowToExpand, const bool bNewExpansion) const
-{
-	if (TreeViewPtr.IsValid())
-	{
-		TreeViewPtr->SetItemExpansion(RowToExpand, bNewExpansion);
-	}
-}
-
-void SConsoleVariablesEditorList::SetAllListViewItemsCheckState(const ECheckBoxState InNewState)
-{
-	for (const TSharedPtr<FConsoleVariablesEditorListRow>& Row : TreeViewRootObjects)
-	{
-		Row->SetWidgetCheckedState(InNewState);
-	}
-}
-
-bool SConsoleVariablesEditorList::DoesListHaveCheckedMembers() const
-{
-	for (const TSharedPtr<FConsoleVariablesEditorListRow>& Row : TreeViewRootObjects)
-	{
-		if (Row->GetCommandInfo().Pin()->ObjectType ==
-			FConsoleVariablesEditorCommandInfo::EConsoleObjectType::Variable &&
-			Row->GetWidgetCheckedState() == ECheckBoxState::Checked)
-		{
-			return true;
-		}
-	}
-
-	return false;
-}
-
-bool SConsoleVariablesEditorList::DoesListHaveUncheckedMembers() const
-{
-	for (const TSharedPtr<FConsoleVariablesEditorListRow>& Row : TreeViewRootObjects)
-	{
-		if (Row->GetCommandInfo().Pin()->ObjectType ==
-			FConsoleVariablesEditorCommandInfo::EConsoleObjectType::Variable &&
-			Row->GetWidgetCheckedState() == ECheckBoxState::Unchecked)
-		{
-			return true;
-		}
-	}
-
-	return false;
-}
-
-void SConsoleVariablesEditorList::OnListItemCheckBoxStateChange(const ECheckBoxState InNewState)
-{
-	HeaderCheckBoxState = ECheckBoxState::Checked;
-
-	if (DoesListHaveUncheckedMembers())
-	{
-		HeaderCheckBoxState = ECheckBoxState::Unchecked;
-
-		if (DoesListHaveCheckedMembers())
-		{
-			HeaderCheckBoxState = ECheckBoxState::Undetermined;
-		}
-	}
-}
-
-void SConsoleVariablesEditorList::ToggleFilterActive(const FString& FilterName)
-{
-	if (const TSharedRef<IConsoleVariablesEditorListFilter>* Match =
-		Algo::FindByPredicate(ShowFilters,
-		[&FilterName](TSharedRef<IConsoleVariablesEditorListFilter> Comparator)
-		{
-			return Comparator->GetFilterName().Equals(FilterName);
-		}))
-	{
-		const TSharedRef<IConsoleVariablesEditorListFilter> Filter = *Match;
-		Filter->ToggleFilterActive();
-
-		EvaluateIfRowsPassFilters();
-	}
-}
-
-void SConsoleVariablesEditorList::EvaluateIfRowsPassFilters()
-{
-	for (const FConsoleVariablesEditorListRowPtr& Row : TreeViewRootObjects)
-	{
-		if (Row.IsValid() && Row->GetRowType() == FConsoleVariablesEditorListRow::SingleCommand)
-		{
-			Row->SetDoesRowPassFilters(Algo::FindByPredicate(
-				ShowFilters,
-				[&Row](const TSharedRef<IConsoleVariablesEditorListFilter>& Filter)
-				{
-					return Filter->GetIsFilterActive() && Filter->DoesItemPassFilter(Row);
-				}) != nullptr);
-		}
-	}
-
-	TreeViewPtr->RequestTreeRefresh();
-}
-
-EColumnSortMode::Type SConsoleVariablesEditorList::GetSortModeForColumn(FName InColumnName) const
-{
-	EColumnSortMode::Type ColumnSortMode = EColumnSortMode::None;
-
-	if (GetActiveSortingColumnName().IsEqual(InColumnName))
-	{
-		ColumnSortMode = ActiveSortingType;
-	}
-
-	return ColumnSortMode;
-}
-
-void SConsoleVariablesEditorList::OnSortColumnCalled(EColumnSortPriority::Type Priority, const FName& ColumnName, EColumnSortMode::Type SortMode)
-{
-	ExecuteSort(ColumnName, CycleSortMode(ColumnName));
-}
-
-EColumnSortMode::Type SConsoleVariablesEditorList::CycleSortMode(const FName& InColumnName)
-{
-	// Custom handler for Custom Sort Order mode
-	if (InColumnName.IsEqual(CustomSortOrderColumnName))
-	{
-		ActiveSortingType = EColumnSortMode::Ascending;
-	}
-	else
-	{
-		const EColumnSortMode::Type PreviousColumnSortMode = GetSortModeForColumn(InColumnName);
-		ActiveSortingType = PreviousColumnSortMode ==
-			EColumnSortMode::Ascending ? EColumnSortMode::Descending : EColumnSortMode::Ascending;
-	}
-
-	ActiveSortingColumnName = InColumnName;
-	return ActiveSortingType;
-}
-
-void SConsoleVariablesEditorList::ExecuteSort(const FName& InColumnName, const EColumnSortMode::Type InColumnSortMode)
+void SConsoleVariablesEditorList::GenerateTreeView()
 {	
-	if (InColumnName.IsEqual(CustomSortOrderColumnName))
+	if (!ensure(TreeViewPtr.IsValid()))
 	{
-		TreeViewRootObjects.StableSort(SortByOrderAscending);
-	}
-	if (InColumnName.IsEqual(SourceColumnName))
-	{
-		TreeViewRootObjects.StableSort(
-			InColumnSortMode == EColumnSortMode::Ascending ? SortBySourceAscending : SortBySourceDescending);
-	}
-	if (InColumnName.IsEqual(VariableNameColumnName))
-	{
-		TreeViewRootObjects.StableSort(
-			InColumnSortMode == EColumnSortMode::Ascending ? SortByVariableNameAscending : SortByVariableNameDescending);
+		return;
 	}
 	
+	FlushMemory(true);
+
+	FConsoleVariablesEditorModule& ConsoleVariablesEditorModule = FConsoleVariablesEditorModule::Get();
+
+	TObjectPtr<UConsoleVariablesAsset> EditableAsset = nullptr;
+
+	switch (ListModelPtr.Pin()->GetListMode()) 
+	{
+		case FConsoleVariablesEditorList::EConsoleVariablesEditorListMode::Preset:
+			EditableAsset = ConsoleVariablesEditorModule.GetPresetAsset();
+			break;
+
+		case FConsoleVariablesEditorList::EConsoleVariablesEditorListMode::GlobalSearch:
+			EditableAsset = ConsoleVariablesEditorModule.GetGlobalSearchAsset();
+			break;
+
+		default:
+			checkNoEntry();
+	}
+	check(EditableAsset);
+
+	for (const FConsoleVariablesEditorAssetSaveData& SavedCommand : EditableAsset->GetSavedCommands())
+	{
+		// Get corresponding CommandInfo for tracking or make one if the command is non-value
+		TSharedPtr<FConsoleVariablesEditorCommandInfo> CommandInfo =
+			ConsoleVariablesEditorModule.FindCommandInfoByName(SavedCommand.CommandName).Pin();
+
+		if (!CommandInfo.IsValid())
+		{
+			CommandInfo = MakeShared<FConsoleVariablesEditorCommandInfo>(SavedCommand.CommandName);
+
+			FConsoleVariablesEditorModule::Get().AddConsoleObjectCommandInfoToMasterReference(
+				CommandInfo.ToSharedRef());
+		}
+		
+		if (CommandInfo.IsValid())
+		{
+			FConsoleVariablesEditorListRowPtr RowToAdd = nullptr;
+			
+			if (TSharedPtr<FConsoleVariablesEditorListRow>* MatchingRow = Algo::FindByPredicate(
+				LastPresetObjects,
+				[this, &CommandInfo](const FConsoleVariablesEditorListRowPtr& ExistingRow)
+				{
+					return ExistingRow->GetCommandInfo().IsValid() &&
+						ExistingRow->GetCommandInfo().Pin()->Command.Equals(CommandInfo->Command);
+				}))
+			{
+				RowToAdd = *MatchingRow;
+			}
+			else
+			{
+				const ECheckBoxState NewCheckedState =
+					SavedCommand.CheckedState == ECheckBoxState::Unchecked ?
+							ECheckBoxState::Unchecked : ECheckBoxState::Checked;
+
+				if (ListModelPtr.Pin()->GetListMode() == FConsoleVariablesEditorList::EConsoleVariablesEditorListMode::Preset)
+				{
+					// If the row is checked and the saved value differs from the current value,
+					// execute the command with the saved value 
+					if (CommandInfo->ObjectType == FConsoleVariablesEditorCommandInfo::EConsoleObjectType::Command ||
+						(NewCheckedState == ECheckBoxState::Checked &&
+						CommandInfo->IsCurrentValueDifferentFromInputValue(SavedCommand.CommandValueAsString)))
+					{
+						if (!SavedCommand.CommandValueAsString.IsEmpty())
+						{
+							CommandInfo->ExecuteCommand(SavedCommand.CommandValueAsString);
+						}
+					}
+				}
+		
+				RowToAdd = 
+				   MakeShared<FConsoleVariablesEditorListRow>(
+						   CommandInfo, SavedCommand.CommandValueAsString,
+						   FConsoleVariablesEditorListRow::SingleCommand, 
+						   NewCheckedState, SharedThis(this), TreeViewRootObjects.Num(), nullptr);
+			}
+			
+			TreeViewRootObjects.Add(RowToAdd);
+		}
+	}
+
+	// Now clear out the last preset cache if the list is in preset mode
+	if (ListModelPtr.Pin()->GetListMode() == FConsoleVariablesEditorList::EConsoleVariablesEditorListMode::Preset)
+	{
+		LastPresetObjects.Empty();
+	}
+
 	TreeViewPtr->RequestTreeRefresh();
 }
 
+void SConsoleVariablesEditorList::FindVisibleTreeViewObjects()
+{
+	VisibleTreeViewObjects.Empty();
+
+	for (const TSharedPtr<FConsoleVariablesEditorListRow>& Row : TreeViewRootObjects)
+	{
+		if (Row->ShouldBeVisible())
+		{
+			VisibleTreeViewObjects.Add(Row);
+		}
+	}
+}
+
+void SConsoleVariablesEditorList::FindVisibleObjectsAndRequestTreeRefresh()
+{
+	FindVisibleTreeViewObjects();
+	TreeViewPtr->RequestTreeRefresh();
+}
+
@@ -772,6 +1082,6 @@
 			SetChildExpansionRecursively(Child, bNewIsExpanded);
 		}
 	}
-};
+}
 
 #undef LOCTEXT_NAMESPACE
diff --git a/Engine/Plugins/Experimental/ConsoleVariables/Source/ConsoleVariablesEditor/Private/Views/List/SConsoleVariablesEditorList.h b/Engine/Plugins/Experimental/ConsoleVariables/Source/ConsoleVariablesEditor/Private/Views/List/SConsoleVariablesEditorList.h
--- a/Engine/Plugins/Experimental/ConsoleVariables/Source/ConsoleVariablesEditor/Private/Views/List/SConsoleVariablesEditorList.h
+++ b/Engine/Plugins/Experimental/ConsoleVariables/Source/ConsoleVariablesEditor/Private/Views/List/SConsoleVariablesEditorList.h
@@ -6,9 +6,13 @@
 #include "ConsoleVariablesEditorListFilters/IConsoleVariablesEditorListFilter.h"
 
 #include "Widgets/SCompoundWidget.h"
+#include "Widgets/Layout/SGridPanel.h"
 #include "Widgets/Views/SHeaderRow.h"
 #include "Widgets/Views/STreeView.h"
 
+class SConsoleVariablesEditorGlobalSearchToggle;
+class SWrapBox;
+class FConsoleVariablesEditorList;
 class UConsoleVariablesAsset;
 class SBox;
 class SComboButton;
@@ -25,10 +29,21 @@ public:
 
 	SLATE_END_ARGS()
 
-	void Construct(const FArguments& InArgs);
+	void Construct(const FArguments& InArgs, TSharedRef<FConsoleVariablesEditorList> ListModel);
 
 	virtual ~SConsoleVariablesEditorList() override;
 
+	TWeakPtr<FConsoleVariablesEditorList> GetListModelPtr()
+	{
+		return ListModelPtr;
+	}
+
+	FReply TryEnterGlobalSearch(const FString& SearchString = "");
+	FReply HandleRemoveGlobalSearchToggleButton();
+	void CleanUpGlobalSearchesMarkedForDelete();
+	/** Remove all widgets from the global search container then recreate them */
+	void RefreshGlobalSearchWidgets();
+
 	/** Regenerate the list items and refresh the list. Call when adding or removing variables. */
 	void RebuildList(const FString& InConsoleCommandToScrollTo = "");
 
@@ -52,7 +67,7 @@ public:
 	void UpdatePresetValuesForSave(const TObjectPtr<UConsoleVariablesAsset> InAsset) const;
 
 	FString GetSearchStringFromSearchInputField() const;
-	void ExecuteListViewSearchOnAllRows(const FString& SearchString) const;
+	void ExecuteListViewSearchOnAllRows(const FString& SearchString, const bool bShouldRefreshAfterward = true);
 
 	bool DoesTreeViewHaveVisibleChildren() const;
 
@@ -67,7 +82,7 @@ public:
 	void OnListItemCheckBoxStateChange(const ECheckBoxState InNewState);
 
 	void ToggleFilterActive(const FString& FilterName);
-	void EvaluateIfRowsPassFilters();
+	void EvaluateIfRowsPassFilters(const bool bShouldRefreshAfterward = true);
 
 	// Sorting
 
@@ -78,13 +93,14 @@ public:
 	EColumnSortMode::Type GetSortModeForColumn(FName InColumnName) const;
 	void OnSortColumnCalled(EColumnSortPriority::Type Priority, const FName& ColumnName, EColumnSortMode::Type SortMode);
 	EColumnSortMode::Type CycleSortMode(const FName& InColumnName);
-	void ExecuteSort(const FName& InColumnName, const EColumnSortMode::Type InColumnSortMode);
+	void ExecuteSort(
+		const FName& InColumnName, const EColumnSortMode::Type InColumnSortMode, const bool bShouldRefreshAfterward = true);
 	void ClearSorting()
 	{
 		ActiveSortingColumnName = NAME_None;
 		ActiveSortingType = EColumnSortMode::None;
 	}
-	void SetSortOrder();
+	void SetSortOrder(const bool bShouldRefreshAfterward = true);
 
 	// Column Names
 
@@ -95,6 +111,8 @@ public:
 	static const FName SourceColumnName;
 
 private:
+
+	TWeakPtr<FConsoleVariablesEditorList> ListModelPtr;
 	
 	TSharedPtr<SHeaderRow> HeaderRow;
 	TSharedPtr<SHeaderRow> GenerateHeaderRow();
@@ -110,15 +128,25 @@ private:
 
 	// Search
 	
-	void OnListViewSearchTextChanged(const FText& Text) const;
+	void OnListViewSearchTextChanged(const FText& Text);
 
 	TSharedPtr<SSearchBox> ListSearchBoxPtr;
 	TSharedPtr<SComboButton> ViewOptionsComboButton;
+	
+	/** Contains the GlobalSearchesContainer and remove button. Made a member in order to collapse it when GlobalSearchesContainer has no children. */
+	TSharedPtr<SHorizontalBox> GlobalSearchesHBox;
+	/** Contains the global search checkboxes only */
+	TSharedPtr<SWrapBox> GlobalSearchesContainer;
+	TArray<TSharedRef<SConsoleVariablesEditorGlobalSearchToggle>> CurrentGlobalSearches;
+	
+	TSharedPtr<SButton> RemoveGlobalSearchesButtonPtr;
 	TSharedPtr<SBox> ListBoxContainerPtr;
 
 	//  Tree View Implementation
 
 	void GenerateTreeView();
+	void FindVisibleTreeViewObjects();
+	void FindVisibleObjectsAndRequestTreeRefresh();
 	
 	void OnGetRowChildren(FConsoleVariablesEditorListRowPtr Row, TArray<FConsoleVariablesEditorListRowPtr>& OutChildren) const;
 	void OnRowChildExpansionChange(FConsoleVariablesEditorListRowPtr Row, const bool bIsExpanded, const bool bIsRecursive = false) const;
@@ -128,8 +156,15 @@ private:
 	TArray<TSharedRef<IConsoleVariablesEditorListFilter>> ShowFilters;
 
 	TSharedPtr<STreeView<FConsoleVariablesEditorListRowPtr>> TreeViewPtr;
-	
+
+	/** All Tree view objects */
 	TArray<FConsoleVariablesEditorListRowPtr> TreeViewRootObjects;
+	/** Visible Tree view objects, after filters */
+	TArray<FConsoleVariablesEditorListRowPtr> VisibleTreeViewObjects;
+	/** When a global search is initiated we cache the tree objects shown in preset list mode
+	 * This way they can be recalled without rebuilding from the saved commands
+	 */
+	TArray<FConsoleVariablesEditorListRowPtr> LastPresetObjects;
 
 	// Sorting
 
diff --git a/Engine/Plugins/Experimental/ConsoleVariables/Source/ConsoleVariablesEditor/Private/Views/List/SConsoleVariablesEditorListRow.cpp b/Engine/Plugins/Experimental/ConsoleVariables/Source/ConsoleVariablesEditor/Private/Views/List/SConsoleVariablesEditorListRow.cpp
--- a/Engine/Plugins/Experimental/ConsoleVariables/Source/ConsoleVariablesEditor/Private/Views/List/SConsoleVariablesEditorListRow.cpp
+++ b/Engine/Plugins/Experimental/ConsoleVariables/Source/ConsoleVariablesEditor/Private/Views/List/SConsoleVariablesEditorListRow.cpp
@@ -2,6 +2,7 @@
 
 #include "SConsoleVariablesEditorListRow.h"
 
+#include "ConsoleVariablesEditorList.h"
 #include "ConsoleVariablesEditorModule.h"
 #include "ConsoleVariablesEditorStyle.h"
 #include "SConsoleVariablesEditorListValueInput.h"
@@ -31,7 +32,8 @@ public:
 	{
 		check(InItems.Num() > 0);
 
-		TSharedRef<FConsoleVariablesListRowDragDropOp> Operation = MakeShareable(new FConsoleVariablesListRowDragDropOp());
+		TSharedRef<FConsoleVariablesListRowDragDropOp> Operation = MakeShareable(
+			new FConsoleVariablesListRowDragDropOp());
 
 		Operation->DraggedItems = InItems;
 
@@ -58,15 +60,16 @@ public:
 };
 
 void SConsoleVariablesEditorListRow::Construct(
-	const FArguments& InArgs, const TSharedRef<STableViewBase>& InOwnerTable, const TWeakPtr<FConsoleVariablesEditorListRow> InRow)
+	const FArguments& InArgs, const TSharedRef<STableViewBase>& InOwnerTable,
+	const TWeakPtr<FConsoleVariablesEditorListRow> InRow)
 {
 	check(InRow.IsValid());
 
 	Item = InRow;
 	const FConsoleVariablesEditorListRowPtr PinnedItem = Item.Pin();
 	const TSharedPtr<FConsoleVariablesEditorCommandInfo> PinnedCommand = PinnedItem->GetCommandInfo().Pin();
-	
-	check (PinnedCommand.IsValid());
+
+	check(PinnedCommand.IsValid());
 
 	// Set up flash animation
 	FlashAnimation = FCurveSequence(0.f, FlashAnimationDuration, ECurveEaseFunction::QuadInOut);
@@ -98,32 +101,34 @@ TSharedRef<SWidget> SConsoleVariablesEditorListRow::GenerateWidgetForColumn(cons
 	const TSharedRef<SImage> FlashImage = SNew(SImage)
 										.Image(new FSlateColorBrush(FStyleColors::White))
 										.Visibility_Raw(this, &SConsoleVariablesEditorListRow::GetFlashImageVisibility)
-										.ColorAndOpacity_Raw(this, &SConsoleVariablesEditorListRow::GetFlashImageColorAndOpacity);
+										.ColorAndOpacity_Raw(
+		                                                  this,
+		                                                  &SConsoleVariablesEditorListRow::GetFlashImageColorAndOpacity);
 
 	FlashImages.Add(FlashImage);
 
 	return SNew(SBox)
+		.Visibility(EVisibility::SelfHitTestInvisible)
+		[
+			SNew(SOverlay)
 			.Visibility(EVisibility::SelfHitTestInvisible)
-			[
-				SNew(SOverlay)
-				.Visibility(EVisibility::SelfHitTestInvisible)
 
-				+SOverlay::Slot()
-				[
-					FlashImage
-				]
+			+ SOverlay::Slot()
+			[
+				FlashImage
+			]
 
-				+SOverlay::Slot()
-				[
-					SNew(SBorder)
+			+ SOverlay::Slot()
+			[
+				SNew(SBorder)
 					.HAlign(HAlign_Fill)
 					.VAlign(VAlign_Center)
 					.BorderImage(GetBorderImage(PinnedItem->GetRowType()))
-					[
-						CellWidget.ToSharedRef()
-					]
+				[
+					CellWidget.ToSharedRef()
 				]
-			];
+			]
+		];
 }
 
 void SConsoleVariablesEditorListRow::OnMouseEnter(const FGeometry& MyGeometry, const FPointerEvent& MouseEvent)
@@ -132,21 +137,43 @@ void SConsoleVariablesEditorListRow::OnMouseEnter(const FGeometry& MyGeometry, c
 
 	if (HoverableWidgetsPtr.IsValid())
 	{
-		HoverableWidgetsPtr->SetVisibility(EVisibility::SelfHitTestInvisible);
+		// We don't want to show the widget if we're in Global Search mode and the command is already in the preset
+		const bool bIsListValid =
+			Item.Pin()->GetListViewPtr().IsValid() && Item.Pin()->GetListViewPtr().Pin()->GetListModelPtr().IsValid();
+
+		const FConsoleVariablesEditorList::EConsoleVariablesEditorListMode ListMode =
+			bIsListValid
+				? Item.Pin()->GetListViewPtr().Pin()->GetListModelPtr().Pin()->GetListMode()
+				: FConsoleVariablesEditorList::EConsoleVariablesEditorListMode::Preset;
+
+		const bool bIsGlobalSearch =
+			ListMode == FConsoleVariablesEditorList::EConsoleVariablesEditorListMode::GlobalSearch;
+
+		const FConsoleVariablesEditorModule& ConsoleVariablesEditorModule = FConsoleVariablesEditorModule::Get();
+
+		const FString& CommandName = Item.Pin()->GetCommandInfo().Pin()->Command;
+		FConsoleVariablesEditorAssetSaveData MatchingData;
+		if (!bIsGlobalSearch ||
+			(bIsGlobalSearch &&
+				!ConsoleVariablesEditorModule.GetPresetAsset()->
+				                              FindSavedDataByCommandString(CommandName, MatchingData)))
+		{
+			HoverableWidgetsPtr->SetVisibility(EVisibility::SelfHitTestInvisible);
+		}
 	}
-	
+
 	SMultiColumnTableRow<FConsoleVariablesEditorListRowPtr>::OnMouseEnter(MyGeometry, MouseEvent);
 }
 
 void SConsoleVariablesEditorListRow::OnMouseLeave(const FPointerEvent& MouseEvent)
 {
 	bIsHovered = false;
-	
+
 	if (HoverableWidgetsPtr.IsValid())
 	{
 		HoverableWidgetsPtr->SetVisibility(EVisibility::Collapsed);
 	}
-	
+
 	SMultiColumnTableRow<FConsoleVariablesEditorListRowPtr>::OnMouseLeave(MouseEvent);
 }
 
@@ -158,7 +185,7 @@ SConsoleVariablesEditorListRow::~SConsoleVariablesEditorListRow()
 	FlashImages.Empty();
 
 	ValueChildInputWidget.Reset();
-	
+
 	HoverableWidgetsPtr.Reset();
 }
 
@@ -167,22 +194,24 @@ FReply SConsoleVariablesEditorListRow::HandleDragDetected(const FGeometry& MyGeo
 	TArray<FConsoleVariablesEditorListRowPtr> DraggedItems = Item.Pin()->GetSelectedTreeViewItems();
 	TSharedRef<FConsoleVariablesListRowDragDropOp> Operation =
 		FConsoleVariablesListRowDragDropOp::New(DraggedItems);
-	
+
 	return FReply::Handled().BeginDragDrop(Operation);
 }
 
 void SConsoleVariablesEditorListRow::HandleDragLeave(const FDragDropEvent& DragDropEvent)
 {
 	if (TSharedPtr<FConsoleVariablesListRowDragDropOp> Operation =
-	DragDropEvent.GetOperationAs<FConsoleVariablesListRowDragDropOp>())
+		DragDropEvent.GetOperationAs<FConsoleVariablesListRowDragDropOp>())
 	{
 		Operation->ResetToDefaultToolTip();
 	}
 }
 
 TOptional<EItemDropZone> SConsoleVariablesEditorListRow::HandleCanAcceptDrop(const FDragDropEvent& DragDropEvent,
-	EItemDropZone DropZone, FConsoleVariablesEditorListRowPtr TargetItem)
-{	
+                                                                             EItemDropZone DropZone,
+                                                                             FConsoleVariablesEditorListRowPtr
+                                                                             TargetItem)
+{
 	TSharedPtr<FConsoleVariablesListRowDragDropOp> Operation =
 		DragDropEvent.GetOperationAs<FConsoleVariablesListRowDragDropOp>();
 
@@ -190,7 +219,7 @@ TOptional<EItemDropZone> SConsoleVariablesEditorListRow::HandleCanAcceptDrop(con
 	{
 		return TOptional<EItemDropZone>();
 	}
-	
+
 	const bool bShouldEvaluateDrop = Item.IsValid() && Item.Pin()->GetListViewPtr().IsValid() &&
 		Item.Pin()->GetListViewPtr().Pin()->GetActiveSortingColumnName().IsEqual(
 			SConsoleVariablesEditorList::CustomSortOrderColumnName);
@@ -198,10 +227,10 @@ TOptional<EItemDropZone> SConsoleVariablesEditorListRow::HandleCanAcceptDrop(con
 	if (!bShouldEvaluateDrop)
 	{
 		Operation->SetToolTip(
-			LOCTEXT("SortByCustomOrderDrgDropWarning","Sort by custom order (\"#\") to drag & drop"),
+			LOCTEXT("SortByCustomOrderDrgDropWarning", "Sort by custom order (\"#\") to drag & drop"),
 			FAppStyle::Get().GetBrush("Graph.ConnectorFeedback.Error")
 		);
-			
+
 		return TOptional<EItemDropZone>();
 	}
 
@@ -213,7 +242,7 @@ TOptional<EItemDropZone> SConsoleVariablesEditorListRow::HandleCanAcceptDrop(con
 	if (bIsDropDenied)
 	{
 		Operation->ResetToDefaultToolTip();
-		
+
 		return TOptional<EItemDropZone>();
 	}
 
@@ -226,12 +255,12 @@ TOptional<EItemDropZone> SConsoleVariablesEditorListRow::HandleCanAcceptDrop(con
 
 	const FText DropPermittedText =
 		FText::Format(InsertFormatText,
-			ItemNameText,
-			DropZone == EItemDropZone::BelowItem ? BelowText : AboveText,
-			FText::FromString(TargetItem->GetCommandInfo().Pin()->Command)
+		              ItemNameText,
+		              DropZone == EItemDropZone::BelowItem ? BelowText : AboveText,
+		              FText::FromString(TargetItem->GetCommandInfo().Pin()->Command)
 		);
-	
-		Operation->SetToolTip(
+
+	Operation->SetToolTip(
 		DropPermittedText,
 		FAppStyle::Get().GetBrush("Graph.ConnectorFeedback.OK")
 	);
@@ -241,7 +270,7 @@ TOptional<EItemDropZone> SConsoleVariablesEditorListRow::HandleCanAcceptDrop(con
 }
 
 FReply SConsoleVariablesEditorListRow::HandleAcceptDrop(const FDragDropEvent& DragDropEvent, EItemDropZone DropZone,
-	FConsoleVariablesEditorListRowPtr TargetItem)
+                                                        FConsoleVariablesEditorListRowPtr TargetItem)
 {
 	TSharedPtr<FConsoleVariablesListRowDragDropOp> Operation =
 		DragDropEvent.GetOperationAs<FConsoleVariablesListRowDragDropOp>();
@@ -252,11 +281,11 @@ FReply SConsoleVariablesEditorListRow::HandleAcceptDrop(const FDragDropEvent& Dr
 	}
 
 	const TSharedPtr<SConsoleVariablesEditorList> ListView = Item.Pin()->GetListViewPtr().Pin();
-	
+
 	TArray<FConsoleVariablesEditorListRowPtr> DraggedItems = Operation->DraggedItems;
 
 	TArray<FConsoleVariablesEditorListRowPtr> AllTreeItemsCopy = ListView->GetTreeViewItems();
-	
+
 	for (FConsoleVariablesEditorListRowPtr DraggedItem : DraggedItems)
 	{
 		if (!DraggedItem.IsValid() || !AllTreeItemsCopy.Contains(DraggedItem))
@@ -266,7 +295,7 @@ FReply SConsoleVariablesEditorListRow::HandleAcceptDrop(const FDragDropEvent& Dr
 
 		AllTreeItemsCopy.Remove(DraggedItem);
 	}
-	
+
 	const int32 TargetIndex = AllTreeItemsCopy.IndexOfByKey(TargetItem);
 
 	if (TargetIndex > -1)
@@ -274,15 +303,15 @@ FReply SConsoleVariablesEditorListRow::HandleAcceptDrop(const FDragDropEvent& Dr
 		for (int32 ItemIndex = DraggedItems.Num() - 1; ItemIndex >= 0; ItemIndex--)
 		{
 			FConsoleVariablesEditorListRowPtr DraggedItem = DraggedItems[ItemIndex];
-			
+
 			if (!DraggedItem.IsValid() || AllTreeItemsCopy.Contains(DraggedItem))
 			{
 				continue;
 			}
-			
+
 			AllTreeItemsCopy.Insert(DraggedItem, DropZone == EItemDropZone::AboveItem ? TargetIndex : TargetIndex + 1);
 		}
-		
+
 		ListView->SetTreeViewItems(AllTreeItemsCopy);
 		ListView->SetSortOrder();
 	}
@@ -309,14 +338,14 @@ FSlateColor SConsoleVariablesEditorListRow::GetFlashImageColorAndOpacity() const
 		return FLinearColor::LerpUsingHSV(FLinearColor::Transparent, FlashColor, Progress);
 	}
 
-	return FLinearColor::Transparent; 
+	return FLinearColor::Transparent;
 }
 
 const FSlateBrush* SConsoleVariablesEditorListRow::GetBorderImage(
 	const FConsoleVariablesEditorListRow::EConsoleVariablesEditorListRowType InRowType)
 {
 	switch (InRowType)
-	{							
+	{
 	case FConsoleVariablesEditorListRow::CommandGroup:
 		return FConsoleVariablesEditorStyle::Get().GetBrush("ConsoleVariablesEditor.CommandGroupBorder");
 
@@ -328,42 +357,44 @@ const FSlateBrush* SConsoleVariablesEditorListRow::GetBorderImage(
 	}
 }
 
-TSharedRef<SWidget> SConsoleVariablesEditorListRow::GenerateCells(const FName& InColumnName, const TSharedPtr<FConsoleVariablesEditorListRow> PinnedItem)
+TSharedRef<SWidget> SConsoleVariablesEditorListRow::GenerateCells(const FName& InColumnName,
+                                                                  const TSharedPtr<FConsoleVariablesEditorListRow>
+                                                                  PinnedItem)
 {
-	
 	if (InColumnName.IsEqual(SConsoleVariablesEditorList::CustomSortOrderColumnName))
 	{
-		return  SNew(STextBlock)
+		return SNew(STextBlock)
 				.Visibility(EVisibility::SelfHitTestInvisible)
 				.Justification(ETextJustify::Center)
 				.Text_Lambda([this, PinnedItem]()
-				{
-					return FText::AsNumber(PinnedItem->GetSortOrder() + 1);
-				});
+		                       {
+			                       return FText::AsNumber(PinnedItem->GetSortOrder() + 1);
+		                       });
 	}
 	if (InColumnName.IsEqual(SConsoleVariablesEditorList::CheckBoxColumnName))
 	{
 		return SNew(SBox)
 				.Visibility(EVisibility::SelfHitTestInvisible)
 				.HAlign(HAlign_Center)
-				[
-					SNew(SCheckBox)
+		[
+			SNew(SCheckBox)
 					.IsChecked_Raw(this, &SConsoleVariablesEditorListRow::GetCheckboxState)
 					.OnCheckStateChanged_Raw(this, &SConsoleVariablesEditorListRow::OnCheckboxStateChange)
 					.Visibility(Item.Pin()->GetCommandInfo().Pin()->ObjectType ==
-						FConsoleVariablesEditorCommandInfo::EConsoleObjectType::Variable ?
-						EVisibility::Visible : EVisibility::Collapsed)
-				];
+					            FConsoleVariablesEditorCommandInfo::EConsoleObjectType::Variable
+						            ? EVisibility::Visible
+						            : EVisibility::Collapsed)
+		];
 	}
 	if (InColumnName.IsEqual(SConsoleVariablesEditorList::VariableNameColumnName))
 	{
 		if (!HoverToolTip.IsValid())
 		{
 			HoverToolTip = SConsoleVariablesEditorTooltipWidget::MakeTooltip(
-							PinnedItem->GetCommandInfo().Pin()->Command,
-							PinnedItem->GetCommandInfo().Pin()->GetHelpText());
+				PinnedItem->GetCommandInfo().Pin()->Command,
+				PinnedItem->GetCommandInfo().Pin()->GetHelpText());
 		}
-		return  SNew(STextBlock)
+		return SNew(STextBlock)
 				.Visibility(EVisibility::Visible)
 				.Text(FText::FromString(PinnedItem->GetCommandInfo().Pin()->Command))
 				.ToolTip(HoverToolTip);
@@ -375,27 +406,27 @@ TSharedRef<SWidget> SConsoleVariablesEditorListRow::GenerateCells(const FName& I
 	if (InColumnName.IsEqual(SConsoleVariablesEditorList::SourceColumnName))
 	{
 		return SNew(SOverlay)
-				.Visibility(EVisibility::SelfHitTestInvisible)
+			.Visibility(EVisibility::SelfHitTestInvisible)
 
-				+SOverlay::Slot()
-				.HAlign(HAlign_Left)
-				.VAlign(VAlign_Center)
-				[
-					SNew(STextBlock)
+			+ SOverlay::Slot()
+			  .HAlign(HAlign_Left)
+			  .VAlign(VAlign_Center)
+			[
+				SNew(STextBlock)
 					.Visibility(EVisibility::SelfHitTestInvisible)
 					.Text_Lambda([this]()
-					{
-						return Item.Pin()->GetCommandInfo().Pin()->GetSourceAsText();
-					})
-				]
+				                {
+					                return Item.Pin()->GetCommandInfo().Pin()->GetSourceAsText();
+				                })
+			]
 
-				+SOverlay::Slot()
-				.HAlign(HAlign_Right)
-				.VAlign(VAlign_Center)
-				[
-					SAssignNew(HoverableWidgetsPtr, SConsoleVariablesEditorListRowHoverWidgets, Item)
-					.Visibility(EVisibility::Collapsed)
-				];
+			+ SOverlay::Slot()
+			  .HAlign(HAlign_Right)
+			  .VAlign(VAlign_Center)
+			[
+				SAssignNew(HoverableWidgetsPtr, SConsoleVariablesEditorListRowHoverWidgets, Item)
+				.Visibility(EVisibility::Collapsed)
+			];
 	}
 
 	return SNullWidget::NullWidget;
@@ -409,7 +440,7 @@ ECheckBoxState SConsoleVariablesEditorListRow::GetCheckboxState() const
 	{
 		return ECheckBoxState::Checked;
 	}
-	
+
 	if (ensure(Item.IsValid()))
 	{
 		return Item.Pin()->GetWidgetCheckedState();
@@ -420,12 +451,12 @@ ECheckBoxState SConsoleVariablesEditorListRow::GetCheckboxState() const
 
 void SConsoleVariablesEditorListRow::OnCheckboxStateChange(const ECheckBoxState InNewState) const
 {
-	if (const TSharedPtr<FConsoleVariablesEditorListRow> PinnedItem = Item.Pin())  
+	if (const TSharedPtr<FConsoleVariablesEditorListRow> PinnedItem = Item.Pin())
 	{
 		PinnedItem->SetWidgetCheckedState(InNewState, true);
-		
+
 		if (PinnedItem->GetRowType() == FConsoleVariablesEditorListRow::SingleCommand)
-		{									
+		{
 			if (PinnedItem->IsRowChecked())
 			{
 				PinnedItem->GetCommandInfo().Pin()->ExecuteCommand(PinnedItem->GetCachedValue());
@@ -438,7 +469,8 @@ void SConsoleVariablesEditorListRow::OnCheckboxStateChange(const ECheckBoxState
 				PinnedItem->ResetToStartupValueAndSource();
 
 				FConsoleVariablesEditorModule::Get().SendMultiUserConsoleVariableChange(
-					PinnedItem->GetCommandInfo().Pin()->Command, PinnedItem->GetCommandInfo().Pin()->StartupValueAsString);
+					PinnedItem->GetCommandInfo().Pin()->Command,
+					PinnedItem->GetCommandInfo().Pin()->StartupValueAsString);
 			}
 		}
 
@@ -449,24 +481,27 @@ void SConsoleVariablesEditorListRow::OnCheckboxStateChange(const ECheckBoxState
 	}
 }
 
-TSharedRef<SWidget> SConsoleVariablesEditorListRow::GenerateValueCellWidget(const TSharedPtr<FConsoleVariablesEditorListRow> PinnedItem)
+TSharedRef<SWidget> SConsoleVariablesEditorListRow::GenerateValueCellWidget(
+	const TSharedPtr<FConsoleVariablesEditorListRow> PinnedItem)
 {
 	if (PinnedItem->GetCommandInfo().IsValid())
 	{
 		ValueChildInputWidget = SConsoleVariablesEditorListValueInput::GetInputWidget(Item);
 		const TSharedRef<SHorizontalBox> FinalValueWidget =
 			SNew(SHorizontalBox)
-			.ToolTipText_Lambda([this, PinnedItem] ()
+			.ToolTipText_Lambda([this, PinnedItem]()
 			{
-				if (const TSharedPtr<FConsoleVariablesEditorCommandInfo> PinnedCommand = PinnedItem->GetCommandInfo().Pin();
+				if (const TSharedPtr<FConsoleVariablesEditorCommandInfo> PinnedCommand = PinnedItem->GetCommandInfo().
+						Pin();
 					PinnedCommand.IsValid())
 				{
 					return FText::Format(
-								ValueWidgetToolTipFormatText,
-								FText::FromString(PinnedItem->GetCachedValue()),
-								FText::FromString(PinnedItem->GetPresetValue()),
-								FText::FromString(PinnedCommand->StartupValueAsString),
-								FConsoleVariablesEditorCommandInfo::ConvertConsoleVariableSetByFlagToText(PinnedCommand->StartupSource)
+						ValueWidgetToolTipFormatText,
+						FText::FromString(PinnedItem->GetCachedValue()),
+						FText::FromString(PinnedItem->GetPresetValue()),
+						FText::FromString(PinnedCommand->StartupValueAsString),
+						FConsoleVariablesEditorCommandInfo::ConvertConsoleVariableSetByFlagToText(
+							PinnedCommand->StartupSource)
 					);
 				}
 
@@ -474,16 +509,16 @@ TSharedRef<SWidget> SConsoleVariablesEditorListRow::GenerateValueCellWidget(cons
 			});
 
 		FinalValueWidget->AddSlot()
-		.VAlign(VAlign_Center)
-		.Padding(FMargin(2, 0))
+		                .VAlign(VAlign_Center)
+		                .Padding(FMargin(2, 0))
 		[
 			ValueChildInputWidget.ToSharedRef()
 		];
 
 		FinalValueWidget->AddSlot()
-		.AutoWidth()
-		.VAlign(VAlign_Center)
-		.Padding(FMargin(2, 0))
+		                .AutoWidth()
+		                .VAlign(VAlign_Center)
+		                .Padding(FMargin(2, 0))
 		[
 			SNew(SButton)
 			.IsFocusable(false)
@@ -491,21 +526,28 @@ TSharedRef<SWidget> SConsoleVariablesEditorListRow::GenerateValueCellWidget(cons
 			.ButtonStyle(&FAppStyle::Get().GetWidgetStyle<FButtonStyle>("NoBorder"))
 			.ContentPadding(0)
 			.Visibility_Lambda([this, PinnedItem]()
-			{
-				if (PinnedItem.IsValid() && PinnedItem->GetRowType() == FConsoleVariablesEditorListRow::SingleCommand)
-				{
-					return PinnedItem->IsRowChecked() && PinnedItem->GetCommandInfo().Pin()->IsCurrentValueDifferentFromInputValue(PinnedItem->GetPresetValue()) ?
-						EVisibility::Visible : EVisibility::Collapsed;
-				}
-
-				return EVisibility::Collapsed;
-			})
+			             {
+				             if (PinnedItem.IsValid() && PinnedItem->GetRowType() ==
+					             FConsoleVariablesEditorListRow::SingleCommand)
+				             {
+					             PinnedItem->SetDoesCurrentValueDifferFromPresetValue(
+						             PinnedItem->GetCommandInfo().Pin()->IsCurrentValueDifferentFromInputValue(
+							             PinnedItem->GetPresetValue()));
+
+					             return PinnedItem->IsRowChecked() && PinnedItem->
+					                    DoesCurrentValueDifferFromPresetValue()
+						                    ? EVisibility::Visible
+						                    : EVisibility::Collapsed;
+				             }
+
+				             return EVisibility::Collapsed;
+			             })
 			.OnClicked_Lambda([this, PinnedItem]()
-			{
-				PinnedItem->ResetToPresetValue();
+			             {
+				             PinnedItem->ResetToPresetValue();
 
-				return FReply::Handled();
-			})
+				             return FReply::Handled();
+			             })
 			[
 				SNew(SImage)
 				.Image(FAppStyle::Get().GetBrush("PropertyWindow.DiffersFromDefault"))
@@ -519,51 +561,71 @@ TSharedRef<SWidget> SConsoleVariablesEditorListRow::GenerateValueCellWidget(cons
 	return SNullWidget::NullWidget;
 }
 
-void SConsoleVariablesEditorListRowHoverWidgets::Construct(const FArguments& InArgs, const TWeakPtr<FConsoleVariablesEditorListRow> InRow)
+void SConsoleVariablesEditorListRowHoverWidgets::Construct(const FArguments& InArgs,
+                                                           const TWeakPtr<FConsoleVariablesEditorListRow> InRow)
 {
 	check(InRow.IsValid());
 
 	Item = InRow;
 
+	const bool bIsListValid =
+		Item.Pin()->GetListViewPtr().IsValid() && Item.Pin()->GetListViewPtr().Pin()->GetListModelPtr().IsValid();
+
+	const FConsoleVariablesEditorList::EConsoleVariablesEditorListMode ListMode =
+		bIsListValid
+			? Item.Pin()->GetListViewPtr().Pin()->GetListModelPtr().Pin()->GetListMode()
+			: FConsoleVariablesEditorList::EConsoleVariablesEditorListMode::Preset;
+
+	FText ButtonTooltip = LOCTEXT("RemoveCvarTooltip",
+	                              "Remove cvar from this list and reset its value to the startup value.");
+	FName ButtonImageName = "Icons.Delete";
+
+	if (ListMode == FConsoleVariablesEditorList::EConsoleVariablesEditorListMode::GlobalSearch)
+	{
+		ButtonImageName = "Icons.Star";
+		ButtonTooltip = LOCTEXT("AddCvarToPresetTooltip", "Add this cvar to your current preset.");
+	}
+
 	ChildSlot
 	[
 		// Remove Button
-		SAssignNew(RemoveButtonPtr, SButton)
-		.ToolTipText(LOCTEXT("RemoveCvarTooltip","Remove cvar from this list and reset its value to the startup value."))
+		SAssignNew(ActionButtonPtr, SButton)
+		.ToolTipText(ButtonTooltip)
 		.ButtonColorAndOpacity(FStyleColors::Transparent)
 		.OnClicked_Lambda([this]()
 		{
-			return Item.Pin()->OnRemoveButtonClicked();
+			SetVisibility(EVisibility::Collapsed);
+			return Item.Pin()->OnActionButtonClicked();
 		})
 		[
 			SNew(SImage)
 			.Visibility(EVisibility::SelfHitTestInvisible)
-			.Image(FAppStyle::Get().GetBrush("Icons.Delete"))
+			.Image(FAppStyle::Get().GetBrush(ButtonImageName))
 			.ColorAndOpacity(FSlateColor::UseForeground())
 		]
 	];
 }
 
 void SConsoleVariablesEditorListRowHoverWidgets::OnMouseEnter(const FGeometry& MyGeometry,
-	const FPointerEvent& MouseEvent)
+                                                              const FPointerEvent& MouseEvent)
 {
 	SCompoundWidget::OnMouseEnter(MyGeometry, MouseEvent);
 
-	RemoveButtonPtr->SetBorderBackgroundColor(FLinearColor(0.f, 0.f, 0.f, 0.4f));
+	ActionButtonPtr->SetBorderBackgroundColor(FLinearColor(0.f, 0.f, 0.f, 0.4f));
 }
 
 void SConsoleVariablesEditorListRowHoverWidgets::OnMouseLeave(const FPointerEvent& MouseEvent)
 {
 	SCompoundWidget::OnMouseLeave(MouseEvent);
 
-	RemoveButtonPtr->SetBorderBackgroundColor(FLinearColor(0.f, 0.f, 0.f, 0.f));
+	ActionButtonPtr->SetBorderBackgroundColor(FLinearColor(0.f, 0.f, 0.f, 0.f));
 }
 
 SConsoleVariablesEditorListRowHoverWidgets::~SConsoleVariablesEditorListRowHoverWidgets()
 {
-	RemoveButtonPtr.Reset();
-
 	Item.Reset();
+	
+	ActionButtonPtr.Reset();
 }
 
 #undef LOCTEXT_NAMESPACE
diff --git a/Engine/Plugins/Experimental/ConsoleVariables/Source/ConsoleVariablesEditor/Private/Views/List/SConsoleVariablesEditorListRow.h b/Engine/Plugins/Experimental/ConsoleVariables/Source/ConsoleVariablesEditor/Private/Views/List/SConsoleVariablesEditorListRow.h
--- a/Engine/Plugins/Experimental/ConsoleVariables/Source/ConsoleVariablesEditor/Private/Views/List/SConsoleVariablesEditorListRow.h
+++ b/Engine/Plugins/Experimental/ConsoleVariables/Source/ConsoleVariablesEditor/Private/Views/List/SConsoleVariablesEditorListRow.h
@@ -108,5 +108,5 @@ private:
 	
 	TWeakPtr<FConsoleVariablesEditorListRow> Item;
 	
-	TSharedPtr<SButton> RemoveButtonPtr;
+	TSharedPtr<SButton> ActionButtonPtr;
 };
diff --git a/Engine/Plugins/Experimental/ConsoleVariables/Source/ConsoleVariablesEditor/Private/Views/List/SConsoleVariablesEditorListValueInput.cpp b/Engine/Plugins/Experimental/ConsoleVariables/Source/ConsoleVariablesEditor/Private/Views/List/SConsoleVariablesEditorListValueInput.cpp
--- a/Engine/Plugins/Experimental/ConsoleVariables/Source/ConsoleVariablesEditor/Private/Views/List/SConsoleVariablesEditorListValueInput.cpp
+++ b/Engine/Plugins/Experimental/ConsoleVariables/Source/ConsoleVariablesEditor/Private/Views/List/SConsoleVariablesEditorListValueInput.cpp
@@ -120,6 +120,11 @@ void SConsoleVariablesEditorListValueInput_Float::Construct(const FArguments& In
 	Item.Pin()->SetCachedValue(GetInputValueAsString());
 }
 
+SConsoleVariablesEditorListValueInput_Float::~SConsoleVariablesEditorListValueInput_Float()
+{
+	InputWidget.Reset();
+}
+
 void SConsoleVariablesEditorListValueInput_Float::SetInputValue(const FString& InValueAsString)
 {
 	InputWidget->SetValue(FCString::Atof(*InValueAsString));
@@ -197,6 +202,11 @@ void SConsoleVariablesEditorListValueInput_Int::Construct(const FArguments& InAr
 	Item.Pin()->SetCachedValue(GetInputValueAsString());
 }
 
+SConsoleVariablesEditorListValueInput_Int::~SConsoleVariablesEditorListValueInput_Int()
+{
+	InputWidget.Reset();
+}
+
 void SConsoleVariablesEditorListValueInput_Int::SetInputValue(const FString& InValueAsString)
 {
 	if (InValueAsString.IsNumeric())
@@ -279,6 +289,11 @@ void SConsoleVariablesEditorListValueInput_String::Construct(const FArguments& I
 	Item.Pin()->SetCachedValue(GetInputValueAsString());
 }
 
+SConsoleVariablesEditorListValueInput_String::~SConsoleVariablesEditorListValueInput_String()
+{
+	InputWidget.Reset();
+}
+
 void SConsoleVariablesEditorListValueInput_String::SetInputValue(const FString& InValueAsString)
 {
 	InputWidget->SetText(FText::FromString(InValueAsString));
@@ -340,6 +355,12 @@ void SConsoleVariablesEditorListValueInput_Bool::Construct(const FArguments& InA
 	Item.Pin()->SetCachedValue(GetInputValueAsString());
 }
 
+SConsoleVariablesEditorListValueInput_Bool::~SConsoleVariablesEditorListValueInput_Bool()
+{
+	InputWidget.Reset();
+	ButtonText.Reset();
+}
+
 void SConsoleVariablesEditorListValueInput_Bool::SetInputValue(const FString& InValueAsString)
 {
 	check (Item.IsValid());
@@ -440,6 +461,12 @@ void SConsoleVariablesEditorListValueInput_Command::Construct(const FArguments&
 	Item.Pin()->SetCachedValue(GetInputValueAsString());
 }
 
+SConsoleVariablesEditorListValueInput_Command::~SConsoleVariablesEditorListValueInput_Command()
+{
+	InputWidget.Reset();
+	InputText.Reset();
+}
+
 void SConsoleVariablesEditorListValueInput_Command::SetInputValue(const FString& InValueAsString)
 {
 	InputText->SetText(FText::FromString(InValueAsString));
@@ -450,7 +477,7 @@ FString SConsoleVariablesEditorListValueInput_Command::GetInputValueAsString()
 	return GetInputValue();
 }
 
-FString SConsoleVariablesEditorListValueInput_Command::GetInputValue()
+FString SConsoleVariablesEditorListValueInput_Command::GetInputValue() const
 {
 	return InputText->GetText().ToString();
 }
diff --git a/Engine/Plugins/Experimental/ConsoleVariables/Source/ConsoleVariablesEditor/Private/Views/List/SConsoleVariablesEditorListValueInput.h b/Engine/Plugins/Experimental/ConsoleVariables/Source/ConsoleVariablesEditor/Private/Views/List/SConsoleVariablesEditorListValueInput.h
--- a/Engine/Plugins/Experimental/ConsoleVariables/Source/ConsoleVariablesEditor/Private/Views/List/SConsoleVariablesEditorListValueInput.h
+++ b/Engine/Plugins/Experimental/ConsoleVariables/Source/ConsoleVariablesEditor/Private/Views/List/SConsoleVariablesEditorListValueInput.h
@@ -41,6 +41,8 @@ public:
 
 	void Construct(const FArguments& InArgs, const TWeakPtr<FConsoleVariablesEditorListRow> InRow);
 	
+	virtual ~SConsoleVariablesEditorListValueInput_Float() override;
+	
 	virtual void SetInputValue(const FString& InValueAsString) override;
 	virtual FString GetInputValueAsString() override;
 
@@ -62,6 +64,8 @@ public:
 
 	void Construct(
 		const FArguments& InArgs, const TWeakPtr<FConsoleVariablesEditorListRow> InRow, const bool bIsShowFlag = false);
+
+	virtual ~SConsoleVariablesEditorListValueInput_Int() override;
 	
 	virtual void SetInputValue(const FString& InValueAsString) override;
 	virtual FString GetInputValueAsString() override;
@@ -83,6 +87,8 @@ public:
 	SLATE_END_ARGS()
 
 	void Construct(const FArguments& InArgs, const TWeakPtr<FConsoleVariablesEditorListRow> InRow);
+
+	virtual ~SConsoleVariablesEditorListValueInput_String() override;
 	
 	virtual void SetInputValue(const FString& InValueAsString) override;
 	virtual FString GetInputValueAsString() override;
@@ -104,6 +110,8 @@ public:
 	SLATE_END_ARGS()
 
 	void Construct(const FArguments& InArgs, const TWeakPtr<FConsoleVariablesEditorListRow> InRow);
+
+	virtual ~SConsoleVariablesEditorListValueInput_Bool() override;
 	
 	virtual void SetInputValue(const FString& InValueAsString) override;
 	void SetInputValue(const bool bNewValue);
@@ -139,10 +147,12 @@ public:
 	void Construct(
 		const FArguments& InArgs, const TWeakPtr<FConsoleVariablesEditorListRow> InRow, const FString& InSavedText);
 	
+	virtual ~SConsoleVariablesEditorListValueInput_Command() override;
+	
 	virtual void SetInputValue(const FString& InValueAsString) override;
 	virtual FString GetInputValueAsString() override;
 
-	FString GetInputValue();
+	FString GetInputValue() const;
 
 private:
 	
diff --git a/Engine/Plugins/Experimental/ConsoleVariables/Source/ConsoleVariablesEditor/Private/Views/MainPanel/ConsoleVariablesEditorMainPanel.cpp b/Engine/Plugins/Experimental/ConsoleVariables/Source/ConsoleVariablesEditor/Private/Views/MainPanel/ConsoleVariablesEditorMainPanel.cpp
--- a/Engine/Plugins/Experimental/ConsoleVariables/Source/ConsoleVariablesEditor/Private/Views/MainPanel/ConsoleVariablesEditorMainPanel.cpp
+++ b/Engine/Plugins/Experimental/ConsoleVariables/Source/ConsoleVariablesEditor/Private/Views/MainPanel/ConsoleVariablesEditorMainPanel.cpp
@@ -19,14 +19,16 @@ FConsoleVariablesEditorMainPanel::FConsoleVariablesEditorMainPanel()
 {
 	EditorList = MakeShared<FConsoleVariablesEditorList>();
 	
-	OnConnectionChangedHandle = MultiUserManager.OnConnectionChange().AddRaw(
-		this, &FConsoleVariablesEditorMainPanel::OnConnectionChanged);
-	OnRemoteCVarChangeHandle = MultiUserManager.OnRemoteCVarChange().AddRaw(
-		this, &FConsoleVariablesEditorMainPanel::OnRemoteCvarChange);
+	OnConnectionChangedHandle = MultiUserManager.OnConnectionChange().AddStatic(
+		&FConsoleVariablesEditorMainPanel::OnConnectionChanged);
+	OnRemoteCVarChangeHandle = MultiUserManager.OnRemoteCVarChange().AddStatic(
+		&FConsoleVariablesEditorMainPanel::OnRemoteCvarChange);
 }
 
 FConsoleVariablesEditorMainPanel::~FConsoleVariablesEditorMainPanel()
 {
+	MainPanelWidget.Reset();
+	EditorList.Reset();
 	MultiUserManager.OnConnectionChange().Remove(OnConnectionChangedHandle);
 	MultiUserManager.OnRemoteCVarChange().Remove(OnRemoteCVarChangeHandle);
 }
@@ -48,7 +50,7 @@ FConsoleVariablesEditorModule& FConsoleVariablesEditorMainPanel::GetConsoleVaria
 
 TObjectPtr<UConsoleVariablesAsset> FConsoleVariablesEditorMainPanel::GetEditingAsset()
 {
-	return GetConsoleVariablesModule().GetEditingAsset();
+	return GetConsoleVariablesModule().GetPresetAsset();
 }
 
 void FConsoleVariablesEditorMainPanel::AddConsoleObjectToPreset(
@@ -68,7 +70,10 @@ void FConsoleVariablesEditorMainPanel::AddConsoleObjectToPreset(
 			}
 		);
 
-		RebuildList(bScrollToNewRow ? InConsoleCommand : "");
+		if (GetEditorList().Pin()->GetListMode() == FConsoleVariablesEditorList::EConsoleVariablesEditorListMode::Preset)
+		{
+			RebuildList(bScrollToNewRow ? InConsoleCommand : "");
+		}
 	}
 }
 
diff --git a/Engine/Plugins/Experimental/ConsoleVariables/Source/ConsoleVariablesEditor/Private/Views/MainPanel/ConsoleVariablesEditorMainPanel.h b/Engine/Plugins/Experimental/ConsoleVariables/Source/ConsoleVariablesEditor/Private/Views/MainPanel/ConsoleVariablesEditorMainPanel.h
--- a/Engine/Plugins/Experimental/ConsoleVariables/Source/ConsoleVariablesEditor/Private/Views/MainPanel/ConsoleVariablesEditorMainPanel.h
+++ b/Engine/Plugins/Experimental/ConsoleVariables/Source/ConsoleVariablesEditor/Private/Views/MainPanel/ConsoleVariablesEditorMainPanel.h
@@ -3,6 +3,7 @@
 #pragma once
 
 #include "ConsoleVariablesEditorModule.h"
+#include "Views/List/ConsoleVariablesEditorList.h"
 #include "Widgets/SWidget.h"
 
 class FConsoleVariablesEditorList;
@@ -54,6 +55,11 @@ public:
 		return EditorList;
 	}
 
+	FConsoleVariablesEditorList::EConsoleVariablesEditorListMode GetEditorListMode() const
+	{
+		return EditorList->GetListMode();
+	}
+
 	UE::ConsoleVariables::MultiUser::Private::FManager& GetMultiUserManager()
 	{
 		return MultiUserManager;
@@ -70,8 +76,8 @@ private:
 
 	TSharedPtr<FConsoleVariablesEditorList> EditorList;
 
-	void OnConnectionChanged(EConcertConnectionStatus Status);
-	void OnRemoteCvarChange(const FString InName, const FString InValue);
+	static void OnConnectionChanged(EConcertConnectionStatus Status);
+	static void OnRemoteCvarChange(const FString InName, const FString InValue);
 
 	UE::ConsoleVariables::MultiUser::Private::FManager MultiUserManager;
 	FDelegateHandle OnConnectionChangedHandle;
diff --git a/Engine/Plugins/Experimental/ConsoleVariables/Source/ConsoleVariablesEditor/Private/Views/Widgets/SConsoleVariablesEditorGlobalSearchToggle.cpp b/Engine/Plugins/Experimental/ConsoleVariables/Source/ConsoleVariablesEditor/Private/Views/Widgets/SConsoleVariablesEditorGlobalSearchToggle.cpp
new file mode 100644
--- /dev/null
+++ b/Engine/Plugins/Experimental/ConsoleVariables/Source/ConsoleVariablesEditor/Private/Views/Widgets/SConsoleVariablesEditorGlobalSearchToggle.cpp
@@ -0,0 +1,137 @@
+﻿// Copyright Epic Games, Inc. All Rights Reserved.
+
+#include "SConsoleVariablesEditorGlobalSearchToggle.h"
+
+#include "Styling/AppStyle.h"
+#include "Widgets/Images/SImage.h"
+#include "Widgets/SBoxPanel.h"
+#include "Widgets/Text/STextBlock.h"
+
+#define LOCTEXT_NAMESPACE "ConsoleVariablesEditor"
+
+void SConsoleVariablesEditorGlobalSearchToggle::Construct(const FArguments& Args, const FText& InButtonText)
+{
+	OnToggleCtrlClicked = Args._OnToggleCtrlClicked;
+	OnToggleAltClicked = Args._OnToggleAltClicked;
+	OnToggleMiddleButtonClicked = Args._OnToggleMiddleButtonClicked;
+	OnToggleRightButtonClicked = Args._OnToggleRightButtonClicked;
+	OnToggleClickedOnce = Args._OnToggleClickedOnce;
+
+	CheckedColor = FColor::MakeRandomColor();
+	GlobalSearchText = InButtonText;
+	
+	ChildSlot
+    [
+		SNew(SBorder)
+		.Padding(1.0f)
+		.BorderImage(FAppStyle::Get().GetBrush("ContentBrowser.FilterBackground"))
+		 [
+			SAssignNew( ToggleButtonPtr, SCheckBox )
+			.Style(FAppStyle::Get(), "ContentBrowser.FilterButton")
+			.ToolTipText(LOCTEXT("GlobalSearchPillButtonTooltip","Toggle this search"))
+			.Padding(0.0f)
+			.IsChecked(this, &SConsoleVariablesEditorGlobalSearchToggle::GetCheckedState)
+			.OnCheckStateChanged_Lambda(
+				[this](ECheckBoxState NewCheckboxState)
+			{
+				SetIsButtonChecked(!GetIsToggleChecked());
+
+				if (FSlateApplication::Get().GetModifierKeys().IsControlDown() && OnToggleCtrlClicked.IsBound())
+				{
+					SetIsMarkedForDelete(true);
+					OnToggleCtrlClicked.Execute();
+				}
+				else if (FSlateApplication::Get().GetModifierKeys().IsAltDown() && OnToggleAltClicked.IsBound())
+				{
+					SetIsMarkedForDelete(true);
+					OnToggleAltClicked.Execute();
+				}
+				else if (OnToggleClickedOnce.IsBound())
+				{
+					SetIsMarkedForDelete(false);
+					OnToggleClickedOnce.Execute();
+				}
+			})
+			[
+				SNew(SHorizontalBox)
+				+SHorizontalBox::Slot()
+				.VAlign(VAlign_Center)
+				.AutoWidth()
+				[
+					SNew(SImage)
+					.Image(FAppStyle::Get().GetBrush("ContentBrowser.FilterImage"))
+					.ColorAndOpacity_Lambda([this]()
+					{
+						return GetIsToggleChecked() ? CheckedColor : UncheckedColor;
+					})
+				]
+				+SHorizontalBox::Slot()
+				.Padding(2.f)
+				.VAlign(VAlign_Center)
+				[
+					SNew(STextBlock)
+					.Text(FText::Format(INVTEXT("\"{0}\""), GlobalSearchText)) // Show text in quotes
+				]
+			]
+		]
+    ];
+}
+
+SConsoleVariablesEditorGlobalSearchToggle::~SConsoleVariablesEditorGlobalSearchToggle()
+{
+	OnToggleCtrlClicked.Unbind();
+	OnToggleAltClicked.Unbind();
+	OnToggleMiddleButtonClicked.Unbind();
+	OnToggleRightButtonClicked.Unbind();
+	OnToggleClickedOnce.Unbind();
+
+	ToggleButtonPtr.Reset();
+}
+
+FReply SConsoleVariablesEditorGlobalSearchToggle::OnMouseButtonDown(const FGeometry& MyGeometry, const FPointerEvent& MouseEvent)
+{
+	return FReply::Handled();
+}
+
+FReply SConsoleVariablesEditorGlobalSearchToggle::OnMouseButtonUp( const FGeometry& InMyGeometry, const FPointerEvent& InMouseEvent) 
+{
+	if(InMouseEvent.GetEffectingButton() == EKeys::MiddleMouseButton && OnToggleMiddleButtonClicked.IsBound())
+	{
+		SetIsMarkedForDelete(true);
+		return OnToggleMiddleButtonClicked.Execute().ReleaseMouseCapture();
+	}
+	else if (InMouseEvent.GetEffectingButton() == EKeys::RightMouseButton && OnToggleRightButtonClicked.IsBound())
+	{
+		SetIsMarkedForDelete(true);
+		return OnToggleRightButtonClicked.Execute().ReleaseMouseCapture();
+	}
+		
+	return FReply::Handled().ReleaseMouseCapture();
+}
+
+bool SConsoleVariablesEditorGlobalSearchToggle::GetIsToggleChecked() const
+{
+	return bIsToggleChecked;
+}
+
+ECheckBoxState SConsoleVariablesEditorGlobalSearchToggle::GetCheckedState() const
+{
+	return GetIsToggleChecked() ? ECheckBoxState::Checked : ECheckBoxState::Unchecked;
+}
+
+void SConsoleVariablesEditorGlobalSearchToggle::SetIsButtonChecked(const bool bNewIsButtonChecked)
+{
+	bIsToggleChecked = bNewIsButtonChecked;
+}
+
+bool SConsoleVariablesEditorGlobalSearchToggle::GetIsMarkedForDelete() const
+{
+	return bIsMarkedForDelete;
+}
+
+void SConsoleVariablesEditorGlobalSearchToggle::SetIsMarkedForDelete(const bool bNewMark)
+{
+	bIsMarkedForDelete = bNewMark;
+}
+
+#undef LOCTEXT_NAMESPACE
diff --git a/Engine/Plugins/Experimental/ConsoleVariables/Source/ConsoleVariablesEditor/Private/Views/Widgets/SConsoleVariablesEditorGlobalSearchToggle.h b/Engine/Plugins/Experimental/ConsoleVariables/Source/ConsoleVariablesEditor/Private/Views/Widgets/SConsoleVariablesEditorGlobalSearchToggle.h
new file mode 100644
--- /dev/null
+++ b/Engine/Plugins/Experimental/ConsoleVariables/Source/ConsoleVariablesEditor/Private/Views/Widgets/SConsoleVariablesEditorGlobalSearchToggle.h
@@ -0,0 +1,66 @@
+﻿// Copyright Epic Games, Inc. All Rights Reserved.
+
+#pragma once
+
+#include "CoreMinimal.h"
+
+#include "Framework/SlateDelegates.h"
+#include "Widgets/DeclarativeSyntaxSupport.h"
+#include "Widgets/Input/SCheckBox.h"
+#include "Widgets/SCompoundWidget.h"
+
+class SBorder;
+
+class SConsoleVariablesEditorGlobalSearchToggle : public SCompoundWidget
+{
+public:
+
+	SLATE_BEGIN_ARGS(SConsoleVariablesEditorGlobalSearchToggle)
+	{}
+		SLATE_EVENT(FOnClicked, OnToggleCtrlClicked)
+		SLATE_EVENT(FOnClicked, OnToggleAltClicked)
+		SLATE_EVENT(FOnClicked, OnToggleMiddleButtonClicked)
+		SLATE_EVENT(FOnClicked, OnToggleRightButtonClicked)
+		SLATE_EVENT(FOnClicked, OnToggleClickedOnce)
+	SLATE_END_ARGS();
+
+	void Construct(const FArguments& Args, const FText& InButtonText);
+
+	virtual ~SConsoleVariablesEditorGlobalSearchToggle() override;
+	
+	virtual FReply OnMouseButtonDown(const FGeometry& MyGeometry, const FPointerEvent& MouseEvent) override;
+	virtual FReply OnMouseButtonUp( const FGeometry& InMyGeometry, const FPointerEvent& InMouseEvent ) override;
+
+	
+	[[nodiscard]] bool GetIsToggleChecked() const;
+	[[nodiscard]] ECheckBoxState GetCheckedState() const;
+	void SetIsButtonChecked(const bool bNewIsButtonChecked);
+
+	[[nodiscard]] bool GetIsMarkedForDelete() const;
+	void SetIsMarkedForDelete(const bool bNewMark);
+
+	const FText& GetGlobalSearchText()
+	{
+		return GlobalSearchText;
+	}
+
+private:
+	
+	FOnClicked OnToggleCtrlClicked;
+	FOnClicked OnToggleAltClicked;
+	FOnClicked OnToggleMiddleButtonClicked;
+	FOnClicked OnToggleRightButtonClicked;
+	FOnClicked OnToggleClickedOnce;
+
+	TSharedPtr<SCheckBox> ToggleButtonPtr;
+	bool bIsToggleChecked = true;
+
+	/** When clicked in a special manner, this search button will be marked for deletion
+	 * The responsibility is on external classes to remove the button from the UI the external class controls. */
+	bool bIsMarkedForDelete = false;
+
+	FColor CheckedColor;
+	FColor UncheckedColor = FColor::Black;
+
+	FText GlobalSearchText;
+};
diff --git a/Engine/Plugins/Experimental/ConsoleVariables/Source/ConsoleVariablesEditor/Public/ConsoleVariablesEditorProjectSettings.h b/Engine/Plugins/Experimental/ConsoleVariables/Source/ConsoleVariablesEditor/Public/ConsoleVariablesEditorProjectSettings.h
--- a/Engine/Plugins/Experimental/ConsoleVariables/Source/ConsoleVariablesEditor/Public/ConsoleVariablesEditorProjectSettings.h
+++ b/Engine/Plugins/Experimental/ConsoleVariables/Source/ConsoleVariablesEditor/Public/ConsoleVariablesEditorProjectSettings.h
@@ -19,22 +19,37 @@ public:
 	
 	UConsoleVariablesEditorProjectSettings(const FObjectInitializer& ObjectInitializer)
 	{
+		UncheckedRowDisplayType = EConsoleVariablesEditorRowDisplayType::ShowCurrentValue;
+		
 		bAddAllChangedConsoleVariablesToCurrentPreset = true;
+
+		ChangedConsoleVariableSkipList =
+		{
+			"Editor.ReflectEditorLevelVisibilityWithGame"
+		};
 	}
 
+	/**
+	 *When a row is unchecked, its associated variable's value will be set to the value recorded when the plugin was loaded.
+	 *The value displayed to the user can be configured with this setting, but will not affect the actual applied value.
+	 *ShowCurrentValue displays the actual value currently applied to the variable.
+	 *ShowLastEnteredValue displays the value that will be applied when the row is checked.
+	 */
+	UPROPERTY(Config, EditAnywhere, Category="Console Variables Editor")
+	EConsoleVariablesEditorRowDisplayType UncheckedRowDisplayType;
+
 	/**
 	 *When variables are changed outside the Console Variables Editor, this option will add the variables to the current preset.
 	 *Does not apply to console commands like 'r.SetNearClipPlane' or 'stat fps'
 	 */
-	UPROPERTY(Config, EditAnywhere, Category="Console Variables")
+	UPROPERTY(Config, EditAnywhere, Category="Console Variables Editor")
 	bool bAddAllChangedConsoleVariablesToCurrentPreset;
 
 	/**
-	 *When a row is unchecked, its associated variable's value will be set to the value recorded when the plugin was loaded.
-	 *The value displayed to the user can be configured with this setting, but will not affect the actual applied value.
-	 *ShowCurrentValue displays the actual value currently applied to the variable.
-	 *ShowLastEnteredValue displays the value that will be applied when the row is checked.
+	 * If bAddAllChangedConsoleVariablesToCurrentPreset is true, this list will filter out any matching variables
+	 * changed outside of the Console Variables Editor so they won't be added to the current preset.
+	 * Matching variables explicitly added inside the Console Variables Editor will not be filtered.
 	 */
-	UPROPERTY(Config, EditAnywhere, Category="Console Variables")
-	EConsoleVariablesEditorRowDisplayType UncheckedRowDisplayType;
+	UPROPERTY(Config, EditAnywhere, Category="Console Variables Editor")
+	TSet<FString> ChangedConsoleVariableSkipList;
 };
