Console Variables Editor:

-Fix for: Issue 129506: Compile errors in DeclarativeSyntaxSupport.h

#[fyi] Jason.Walter

diff --git a/Engine/Plugins/Experimental/ConsoleVariables/Source/ConsoleVariablesEditor/Private/Views/Widgets/SConsoleVariablesEditorGlobalSearchToggle.cpp b/Engine/Plugins/Experimental/ConsoleVariables/Source/ConsoleVariablesEditor/Private/Views/Widgets/SConsoleVariablesEditorGlobalSearchToggle.cpp
--- a/Engine/Plugins/Experimental/ConsoleVariables/Source/ConsoleVariablesEditor/Private/Views/Widgets/SConsoleVariablesEditorGlobalSearchToggle.cpp
+++ b/Engine/Plugins/Experimental/ConsoleVariables/Source/ConsoleVariablesEditor/Private/Views/Widgets/SConsoleVariablesEditorGlobalSearchToggle.cpp
@@ -4,6 +4,7 @@
 
 #include "Styling/AppStyle.h"
 #include "Widgets/Images/SImage.h"
+#include "Widgets/Layout/SBorder.h"
 #include "Widgets/SBoxPanel.h"
 #include "Widgets/Text/STextBlock.h"
 
diff --git a/Engine/Plugins/Experimental/ConsoleVariables/Source/ConsoleVariablesEditor/Private/Views/Widgets/SConsoleVariablesEditorGlobalSearchToggle.h b/Engine/Plugins/Experimental/ConsoleVariables/Source/ConsoleVariablesEditor/Private/Views/Widgets/SConsoleVariablesEditorGlobalSearchToggle.h
--- a/Engine/Plugins/Experimental/ConsoleVariables/Source/ConsoleVariablesEditor/Private/Views/Widgets/SConsoleVariablesEditorGlobalSearchToggle.h
+++ b/Engine/Plugins/Experimental/ConsoleVariables/Source/ConsoleVariablesEditor/Private/Views/Widgets/SConsoleVariablesEditorGlobalSearchToggle.h
@@ -9,8 +9,6 @@
 #include "Widgets/Input/SCheckBox.h"
 #include "Widgets/SCompoundWidget.h"
 
-class SBorder;
-
 class SConsoleVariablesEditorGlobalSearchToggle : public SCompoundWidget
 {
 public:
