Console Variables Editor:

-Fix for //UE5/Release-5.0 - NonUnity Compile LyraEditor Win64 - use of undefined type 'FSlateApplication'

#[fyi] Jason.Walter

diff --git a/Engine/Plugins/Experimental/ConsoleVariables/Source/ConsoleVariablesEditor/Private/Views/Widgets/SConsoleVariablesEditorGlobalSearchToggle.cpp b/Engine/Plugins/Experimental/ConsoleVariables/Source/ConsoleVariablesEditor/Private/Views/Widgets/SConsoleVariablesEditorGlobalSearchToggle.cpp
--- a/Engine/Plugins/Experimental/ConsoleVariables/Source/ConsoleVariablesEditor/Private/Views/Widgets/SConsoleVariablesEditorGlobalSearchToggle.cpp
+++ b/Engine/Plugins/Experimental/ConsoleVariables/Source/ConsoleVariablesEditor/Private/Views/Widgets/SConsoleVariablesEditorGlobalSearchToggle.cpp
@@ -2,6 +2,7 @@
 
 #include "SConsoleVariablesEditorGlobalSearchToggle.h"
 
+#include "Framework/Application/SlateApplication.h"
 #include "Styling/AppStyle.h"
 #include "Widgets/Images/SImage.h"
 #include "Widgets/Layout/SBorder.h"
