Console Variables Editor:

-Mitigate log spam when cvar's new priority is lower than its previous priority
-Code formatting polish
-Fix for: Checked states not saved when reopening dialogue or adding new variables to list

#[fyi] Jason.Walter

diff --git a/Engine/Plugins/Experimental/ConsoleVariables/Source/ConsoleVariablesEditor/Private/ConsoleVariablesEditorCommandInfo.h b/Engine/Plugins/Experimental/ConsoleVariables/Source/ConsoleVariablesEditor/Private/ConsoleVariablesEditorCommandInfo.h
--- a/Engine/Plugins/Experimental/ConsoleVariables/Source/ConsoleVariablesEditor/Private/ConsoleVariablesEditorCommandInfo.h
+++ b/Engine/Plugins/Experimental/ConsoleVariables/Source/ConsoleVariablesEditor/Private/ConsoleVariablesEditorCommandInfo.h
@@ -161,6 +161,14 @@ struct FConsoleVariablesEditorCommandInfo
 
 	void SetSourceFlag(const EConsoleVariableFlags InSource)
 	{
+		const uint32 OldPri = (uint32)GetSource();
+		const uint32 NewPri = (uint32)InSource;
+
+		if (NewPri < OldPri)
+		{
+			return;
+		}
+		
 		if (IConsoleObject* ConsoleObject = GetConsoleObjectPtr())
 		{
 			ClearSourceFlags();
diff --git a/Engine/Plugins/Experimental/ConsoleVariables/Source/ConsoleVariablesEditor/Private/Views/List/SConsoleVariablesEditorList.cpp b/Engine/Plugins/Experimental/ConsoleVariables/Source/ConsoleVariablesEditor/Private/Views/List/SConsoleVariablesEditorList.cpp
--- a/Engine/Plugins/Experimental/ConsoleVariables/Source/ConsoleVariablesEditor/Private/Views/List/SConsoleVariablesEditorList.cpp
+++ b/Engine/Plugins/Experimental/ConsoleVariables/Source/ConsoleVariablesEditor/Private/Views/List/SConsoleVariablesEditorList.cpp
@@ -422,9 +422,7 @@ void SConsoleVariablesEditorList::UpdatePresetValuesForSave(const TObjectPtr<UCo
 	
 	for (const FConsoleVariablesEditorListRowPtr& Item : TreeViewRootObjects)
 	{
-		const TWeakPtr<FConsoleVariablesEditorCommandInfo> CommandInfo = Item->GetCommandInfo();
-		
-		if (CommandInfo.IsValid())
+		if (const TWeakPtr<FConsoleVariablesEditorCommandInfo> CommandInfo = Item->GetCommandInfo(); CommandInfo.IsValid())
 		{
 			NewSavedCommands.Add(
 				{
diff --git a/Engine/Plugins/Experimental/ConsoleVariables/Source/ConsoleVariablesEditor/Private/Views/List/SConsoleVariablesEditorListRow.cpp b/Engine/Plugins/Experimental/ConsoleVariables/Source/ConsoleVariablesEditor/Private/Views/List/SConsoleVariablesEditorListRow.cpp
--- a/Engine/Plugins/Experimental/ConsoleVariables/Source/ConsoleVariablesEditor/Private/Views/List/SConsoleVariablesEditorListRow.cpp
+++ b/Engine/Plugins/Experimental/ConsoleVariables/Source/ConsoleVariablesEditor/Private/Views/List/SConsoleVariablesEditorListRow.cpp
@@ -409,20 +409,20 @@ TSharedRef<SWidget> SConsoleVariablesEditorListRow::GenerateCells(const FName& I
 			.Visibility(EVisibility::SelfHitTestInvisible)
 
 			+ SOverlay::Slot()
-			  .HAlign(HAlign_Left)
-			  .VAlign(VAlign_Center)
+			.HAlign(HAlign_Left)
+			.VAlign(VAlign_Center)
 			[
 				SNew(STextBlock)
-					.Visibility(EVisibility::SelfHitTestInvisible)
-					.Text_Lambda([this]()
-				                {
-					                return Item.Pin()->GetCommandInfo().Pin()->GetSourceAsText();
-				                })
+				.Visibility(EVisibility::SelfHitTestInvisible)
+				.Text_Lambda([this]()
+                {
+	                return Item.Pin()->GetCommandInfo().Pin()->GetSourceAsText();
+                })
 			]
 
 			+ SOverlay::Slot()
-			  .HAlign(HAlign_Right)
-			  .VAlign(VAlign_Center)
+		    .HAlign(HAlign_Right)
+		    .VAlign(VAlign_Center)
 			[
 				SAssignNew(HoverableWidgetsPtr, SConsoleVariablesEditorListRowHoverWidgets, Item)
 				.Visibility(EVisibility::Collapsed)
@@ -470,13 +470,25 @@ void SConsoleVariablesEditorListRow::OnCheckboxStateChange(const ECheckBoxState
 
 				FConsoleVariablesEditorModule::Get().SendMultiUserConsoleVariableChange(
 					PinnedItem->GetCommandInfo().Pin()->Command,
-					PinnedItem->GetCommandInfo().Pin()->StartupValueAsString);
+					PinnedItem->GetCommandInfo().Pin()->StartupValueAsString
+				);
 			}
 		}
 
 		if (const TWeakPtr<SConsoleVariablesEditorList> ListView = PinnedItem->GetListViewPtr(); ListView.IsValid())
 		{
-			ListView.Pin()->OnListItemCheckBoxStateChange(InNewState);
+			const TSharedPtr<SConsoleVariablesEditorList> PinnedListView = ListView.Pin();
+			PinnedListView->OnListItemCheckBoxStateChange(InNewState);
+
+			if (const TWeakPtr<FConsoleVariablesEditorList>& ListModel = PinnedListView->GetListModelPtr(); ListModel.IsValid())
+			{
+				const FConsoleVariablesEditorModule& ConsoleVariablesEditorModule = FConsoleVariablesEditorModule::Get();
+				TObjectPtr<UConsoleVariablesAsset> Asset =
+					ListModel.Pin()->GetListMode() == FConsoleVariablesEditorList::EConsoleVariablesEditorListMode::Preset ?
+						ConsoleVariablesEditorModule.GetPresetAsset() : ConsoleVariablesEditorModule.GetGlobalSearchAsset();
+				
+				PinnedListView->UpdatePresetValuesForSave(Asset);
+			}
 		}
 	}
 }
