Console Variables Editor:

-UI Polish pass

-Presets dropdown should not have an outline.
-Check spacing on binoculars and show settings icons.
-Reduce default width of checkbox column to be tighter.
-Row text should not be to the left (slightly) of the column header.
-BG grey between console and search should be the same/lighter color
-Allocate column space per row on far right for remove/star icon.
-MU button is too close to right edge, match Take Recorder.
-Rename tab/menu item to just be "Console Variables."

#[fyi] Jason.Walter

diff --git a/Engine/Plugins/Experimental/ConsoleVariables/Source/ConsoleVariablesEditor/Private/ConsoleVariablesEditorModule.cpp b/Engine/Plugins/Experimental/ConsoleVariables/Source/ConsoleVariablesEditor/Private/ConsoleVariablesEditorModule.cpp
--- a/Engine/Plugins/Experimental/ConsoleVariables/Source/ConsoleVariablesEditor/Private/ConsoleVariablesEditorModule.cpp
+++ b/Engine/Plugins/Experimental/ConsoleVariables/Source/ConsoleVariablesEditor/Private/ConsoleVariablesEditorModule.cpp
@@ -217,7 +217,7 @@ void FConsoleVariablesEditorModule::RegisterMenuItem()
 		ConsoleVariablesToolkitPanelTabId,
 		FOnSpawnTab::CreateRaw(this, & FConsoleVariablesEditorModule::SpawnMainPanelTab))
 			.SetIcon(FSlateIcon(FConsoleVariablesEditorStyle::Get().GetStyleSetName(), "ConsoleVariables.ToolbarButton", "ConsoleVariables.ToolbarButton.Small"))
-			.SetDisplayName(LOCTEXT("OpenConsoleVariablesEditorMenuItem", "Console Variables Editor"))
+			.SetDisplayName(LOCTEXT("OpenConsoleVariablesEditorMenuItem", "Console Variables"))
 			.SetTooltipText(LOCTEXT("OpenConsoleVariablesEditorTooltip", "Open the Console Variables Editor"))
 			.SetMenuType(ETabSpawnerMenuType::Enabled);
 
diff --git a/Engine/Plugins/Experimental/ConsoleVariables/Source/ConsoleVariablesEditor/Private/Views/List/SConsoleVariablesEditorList.cpp b/Engine/Plugins/Experimental/ConsoleVariables/Source/ConsoleVariablesEditor/Private/Views/List/SConsoleVariablesEditorList.cpp
--- a/Engine/Plugins/Experimental/ConsoleVariables/Source/ConsoleVariablesEditor/Private/Views/List/SConsoleVariablesEditorList.cpp
+++ b/Engine/Plugins/Experimental/ConsoleVariables/Source/ConsoleVariablesEditor/Private/Views/List/SConsoleVariablesEditorList.cpp
@@ -25,10 +25,11 @@
 #define LOCTEXT_NAMESPACE "ConsoleVariablesEditor"
 
 const FName SConsoleVariablesEditorList::CustomSortOrderColumnName(TEXT("Order"));
-const FName SConsoleVariablesEditorList::CheckBoxColumnName(TEXT("Column"));
+const FName SConsoleVariablesEditorList::CheckBoxColumnName(TEXT("Checkbox"));
 const FName SConsoleVariablesEditorList::VariableNameColumnName(TEXT("Name"));
 const FName SConsoleVariablesEditorList::ValueColumnName(TEXT("Value"));
 const FName SConsoleVariablesEditorList::SourceColumnName(TEXT("Source"));
+const FName SConsoleVariablesEditorList::ActionButtonColumnName(TEXT("Action"));
 
 void SConsoleVariablesEditorList::Construct(const FArguments& InArgs, TSharedRef<FConsoleVariablesEditorList> ListModel)
 {
@@ -52,19 +53,21 @@ void SConsoleVariablesEditorList::Construct(const FArguments& InArgs, TSharedRef
 		+ SVerticalBox::Slot()
 		.VAlign(VAlign_Top)
 		.AutoHeight()
+		.Padding(FMargin(8.f, 0.f, 8.f, 0.f))
 		[
 			SNew(SHorizontalBox)
 
 			+SHorizontalBox::Slot()
-			.Padding(10.f, 1.f, 0.f, 1.f)
+			.Padding(0.f, 1.f, 0.f, 1.f)
 			[
 				SAssignNew(ListSearchBoxPtr, SSearchBox)
 				.HintText(LOCTEXT("ConsoleVariablesEditorList_SearchHintText", "Search tracked variables, values, sources or help text..."))
 				.OnTextChanged_Raw(this, &SConsoleVariablesEditorList::OnListViewSearchTextChanged)
 			]
 
+			// Global Search
 			+SHorizontalBox::Slot()
-			.Padding(4.f, 0.f, 2.f, 0.f)
+			.Padding(8.f, 1.f, 0.f, 1.f)
 			.VAlign(VAlign_Center)
 			.AutoWidth()
 			[
@@ -84,12 +87,15 @@ void SConsoleVariablesEditorList::Construct(const FArguments& InArgs, TSharedRef
 				]
 			]
 
+			// Show Options
 			+SHorizontalBox::Slot()
-			.AutoWidth()
-			.Padding(10.f, 1.f, 15.f, 1.f)
 			.HAlign(HAlign_Right)
+			.VAlign(VAlign_Center)
+			.AutoWidth()
+			.Padding(8.f, 1.f, 0.f, 1.f)
 			[
 				SAssignNew( ViewOptionsComboButton, SComboButton )
+				.ToolTipText(LOCTEXT("ShowOptions_Tooltip", "Show options to affect the visibility of items in the Console Variables Editor list"))
 				.ComboButtonStyle( FAppStyle::Get(), "SimpleComboButtonWithIcon" ) // Use the tool bar item style for this button
 				.OnGetMenuContent( this, &SConsoleVariablesEditorList::BuildShowOptionsMenu)
 				.HasDownArrow(false)
@@ -699,7 +705,7 @@ TSharedPtr<SHeaderRow> SConsoleVariablesEditorList::GenerateHeaderRow()
 		SHeaderRow::Column(CheckBoxColumnName)
 			.DefaultLabel(LOCTEXT("ConsoleVariablesEditorList_ConsoleVariableCheckboxHeaderText", "Checkbox"))
 			.HAlignHeader(EHorizontalAlignment::HAlign_Center)
-			.FixedWidth(50.f)
+			.FixedWidth(25.f)
 			.ShouldGenerateWidget(true)
 			.HeaderContent()
 			[
@@ -740,6 +746,7 @@ TSharedPtr<SHeaderRow> SConsoleVariablesEditorList::GenerateHeaderRow()
 			.DefaultLabel(LOCTEXT("ConsoleVariablesEditorList_ConsoleVariableValueHeaderText", "Value"))
 			.HAlignHeader(EHorizontalAlignment::HAlign_Left)
 			.ShouldGenerateWidget(true)
+			.FillWidth(0.8f)
 	);
 
 	HeaderRow->AddColumn(
@@ -749,6 +756,18 @@ TSharedPtr<SHeaderRow> SConsoleVariablesEditorList::GenerateHeaderRow()
 			.HAlignHeader(EHorizontalAlignment::HAlign_Left)
 			.SortMode_Raw(this, &SConsoleVariablesEditorList::GetSortModeForColumn, SourceColumnName)
 			.OnSort_Raw(this, &SConsoleVariablesEditorList::OnSortColumnCalled)
+			);
+	
+	HeaderRow->AddColumn(
+		SHeaderRow::Column(ActionButtonColumnName)
+			.DefaultLabel(LOCTEXT("ConsoleVariablesEditorList_ConsoleVariableActionButtonHeaderText", "Action"))
+			.HAlignHeader(EHorizontalAlignment::HAlign_Center)
+			.FixedWidth(25.f)
+			.ShouldGenerateWidget(true)
+			.HeaderContent()
+			[
+				SNew(SBox)
+			]
 	);
 
 	return HeaderRow;
diff --git a/Engine/Plugins/Experimental/ConsoleVariables/Source/ConsoleVariablesEditor/Private/Views/List/SConsoleVariablesEditorList.h b/Engine/Plugins/Experimental/ConsoleVariables/Source/ConsoleVariablesEditor/Private/Views/List/SConsoleVariablesEditorList.h
--- a/Engine/Plugins/Experimental/ConsoleVariables/Source/ConsoleVariablesEditor/Private/Views/List/SConsoleVariablesEditorList.h
+++ b/Engine/Plugins/Experimental/ConsoleVariables/Source/ConsoleVariablesEditor/Private/Views/List/SConsoleVariablesEditorList.h
@@ -109,6 +109,7 @@ public:
 	static const FName VariableNameColumnName;
 	static const FName ValueColumnName;
 	static const FName SourceColumnName;
+	static const FName ActionButtonColumnName;
 
 private:
 
diff --git a/Engine/Plugins/Experimental/ConsoleVariables/Source/ConsoleVariablesEditor/Private/Views/List/SConsoleVariablesEditorListRow.cpp b/Engine/Plugins/Experimental/ConsoleVariables/Source/ConsoleVariablesEditor/Private/Views/List/SConsoleVariablesEditorListRow.cpp
--- a/Engine/Plugins/Experimental/ConsoleVariables/Source/ConsoleVariablesEditor/Private/Views/List/SConsoleVariablesEditorListRow.cpp
+++ b/Engine/Plugins/Experimental/ConsoleVariables/Source/ConsoleVariablesEditor/Private/Views/List/SConsoleVariablesEditorListRow.cpp
@@ -15,6 +15,7 @@
 #include "Widgets/Layout/SBox.h"
 #include "Widgets/SBoxPanel.h"
 #include "Widgets/SOverlay.h"
+#include "Widgets/Layout/SScaleBox.h"
 #include "Widgets/Text/STextBlock.h"
 
 #define LOCTEXT_NAMESPACE "ConsoleVariablesEditor"
@@ -107,10 +108,7 @@ TSharedRef<SWidget> SConsoleVariablesEditorListRow::GenerateWidgetForColumn(cons
 
 	FlashImages.Add(FlashImage);
 
-	return SNew(SBox)
-		.Visibility(EVisibility::SelfHitTestInvisible)
-		[
-			SNew(SOverlay)
+	return SNew(SOverlay)
 			.Visibility(EVisibility::SelfHitTestInvisible)
 
 			+ SOverlay::Slot()
@@ -127,8 +125,7 @@ TSharedRef<SWidget> SConsoleVariablesEditorListRow::GenerateWidgetForColumn(cons
 				[
 					CellWidget.ToSharedRef()
 				]
-			]
-		];
+			];
 }
 
 void SConsoleVariablesEditorListRow::OnMouseEnter(const FGeometry& MyGeometry, const FPointerEvent& MouseEvent)
@@ -394,10 +391,18 @@ TSharedRef<SWidget> SConsoleVariablesEditorListRow::GenerateCells(const FName& I
 				PinnedItem->GetCommandInfo().Pin()->Command,
 				PinnedItem->GetCommandInfo().Pin()->GetHelpText());
 		}
-		return SNew(STextBlock)
-				.Visibility(EVisibility::Visible)
-				.Text(FText::FromString(PinnedItem->GetCommandInfo().Pin()->Command))
-				.ToolTip(HoverToolTip);
+		return SNew(SBox)
+				.Visibility(EVisibility::SelfHitTestInvisible)
+				.HAlign(HAlign_Left)
+				.VAlign(VAlign_Center)
+				.Padding(FMargin(TextBlockLeftPadding, 0.f, 0.f, 0.f))
+				[
+					SNew(STextBlock)
+					.Visibility(EVisibility::Visible)
+					.Justification(ETextJustify::Left)
+					.Text(FText::FromString(PinnedItem->GetCommandInfo().Pin()->Command))
+					.ToolTip(HoverToolTip)
+				];
 	}
 	if (InColumnName.IsEqual(SConsoleVariablesEditorList::ValueColumnName))
 	{
@@ -405,28 +410,23 @@ TSharedRef<SWidget> SConsoleVariablesEditorListRow::GenerateCells(const FName& I
 	}
 	if (InColumnName.IsEqual(SConsoleVariablesEditorList::SourceColumnName))
 	{
-		return SNew(SOverlay)
-			.Visibility(EVisibility::SelfHitTestInvisible)
-
-			+ SOverlay::Slot()
-			.HAlign(HAlign_Left)
-			.VAlign(VAlign_Center)
-			[
-				SNew(STextBlock)
-				.Visibility(EVisibility::SelfHitTestInvisible)
-				.Text_Lambda([this]()
-                {
-	                return Item.Pin()->GetCommandInfo().Pin()->GetSourceAsText();
-                })
-			]
-
-			+ SOverlay::Slot()
-		    .HAlign(HAlign_Right)
-		    .VAlign(VAlign_Center)
-			[
-				SAssignNew(HoverableWidgetsPtr, SConsoleVariablesEditorListRowHoverWidgets, Item)
-				.Visibility(EVisibility::Collapsed)
-			];
+		return SNew(SBox)
+				.Visibility(EVisibility::HitTestInvisible)
+				.HAlign(HAlign_Left)
+				.VAlign(VAlign_Center)
+				.Padding(FMargin(TextBlockLeftPadding, 0.f, 0.f, 0.f))
+				[
+					SNew(STextBlock)
+					.Text_Lambda([this]()
+	                {
+		                return Item.Pin()->GetCommandInfo().Pin()->GetSourceAsText();
+	                })
+				];
+	}
+	if (InColumnName.IsEqual(SConsoleVariablesEditorList::ActionButtonColumnName))
+	{
+		return SAssignNew(HoverableWidgetsPtr, SConsoleVariablesEditorListRowHoverWidgets, Item)
+				.Visibility(EVisibility::Collapsed);
 	}
 
 	return SNullWidget::NullWidget;
@@ -600,20 +600,25 @@ void SConsoleVariablesEditorListRowHoverWidgets::Construct(const FArguments& InA
 
 	ChildSlot
 	[
-		// Remove Button
+		// Action Button
 		SAssignNew(ActionButtonPtr, SButton)
+		.Visibility(EVisibility::Visible)
 		.ToolTipText(ButtonTooltip)
 		.ButtonColorAndOpacity(FStyleColors::Transparent)
+		.ContentPadding(0.f)
 		.OnClicked_Lambda([this]()
 		{
 			SetVisibility(EVisibility::Collapsed);
 			return Item.Pin()->OnActionButtonClicked();
 		})
 		[
-			SNew(SImage)
-			.Visibility(EVisibility::SelfHitTestInvisible)
-			.Image(FAppStyle::Get().GetBrush(ButtonImageName))
-			.ColorAndOpacity(FSlateColor::UseForeground())
+			SNew(SScaleBox)
+			[
+				SNew(SImage)
+				.Visibility(EVisibility::SelfHitTestInvisible)
+				.Image(FAppStyle::Get().GetBrush(ButtonImageName))
+				.ColorAndOpacity(FSlateColor::UseForeground())
+			]
 		]
 	];
 }
diff --git a/Engine/Plugins/Experimental/ConsoleVariables/Source/ConsoleVariablesEditor/Private/Views/List/SConsoleVariablesEditorListRow.h b/Engine/Plugins/Experimental/ConsoleVariables/Source/ConsoleVariablesEditor/Private/Views/List/SConsoleVariablesEditorListRow.h
--- a/Engine/Plugins/Experimental/ConsoleVariables/Source/ConsoleVariablesEditor/Private/Views/List/SConsoleVariablesEditorListRow.h
+++ b/Engine/Plugins/Experimental/ConsoleVariables/Source/ConsoleVariablesEditor/Private/Views/List/SConsoleVariablesEditorListRow.h
@@ -83,6 +83,9 @@ private:
 
 	const float FlashAnimationDuration = 0.75f;
 	const FLinearColor FlashColor = FLinearColor::White;
+
+	/** The offset applied to text widgets so that the text aligns with the column header text */
+	float TextBlockLeftPadding = 3.0f;
 	
 	bool bIsHovered = false;
 };
