Console Variables Editor:

-Fix for: ECVF_RenderThreadSafe cvars seem to revert their values when new value committed with enter key
-Possible fix for: crash on nodes when performing various actions
-Fix for: Value-type commands don't get sent in MU session

diff --git a/Engine/Plugins/Experimental/ConsoleVariables/Source/ConsoleVariablesEditor/Private/ConsoleVariablesAsset.cpp b/Engine/Plugins/Experimental/ConsoleVariables/Source/ConsoleVariablesEditor/Private/ConsoleVariablesAsset.cpp
--- a/Engine/Plugins/Experimental/ConsoleVariables/Source/ConsoleVariablesEditor/Private/ConsoleVariablesAsset.cpp
+++ b/Engine/Plugins/Experimental/ConsoleVariables/Source/ConsoleVariablesEditor/Private/ConsoleVariablesAsset.cpp
@@ -2,6 +2,7 @@
 
 #include "ConsoleVariablesAsset.h"
 
+#include "ConsoleVariablesEditorModule.h"
 #include "Algo/Find.h"
 
 void UConsoleVariablesAsset::SetVariableCollectionDescription(const FString& InVariableCollectionDescription)
@@ -14,7 +15,7 @@ void UConsoleVariablesAsset::ReplaceSavedCommands(const TArray<FConsoleVariables
 	SavedCommands = Replacement;
 }
 
-bool UConsoleVariablesAsset::FindSavedDataByCommandString(const FString& InCommandString, FConsoleVariablesEditorAssetSaveData& OutValue) const
+bool UConsoleVariablesAsset::FindSavedDataByCommandString(const FString InCommandString, FConsoleVariablesEditorAssetSaveData& OutValue) const
 {
 	if (const auto* Match = Algo::FindByPredicate(
 		SavedCommands,
@@ -32,11 +33,28 @@ bool UConsoleVariablesAsset::FindSavedDataByCommandString(const FString& InComma
 
 void UConsoleVariablesAsset::AddOrSetConsoleObjectSavedData(const FConsoleVariablesEditorAssetSaveData& InData)
 {
+	FConsoleVariablesEditorModule& ConsoleVariablesEditorModule = FConsoleVariablesEditorModule::Get();
+
+	if (const TWeakPtr<FConsoleVariablesEditorCommandInfo> CommandInfo =
+			ConsoleVariablesEditorModule.FindCommandInfoByName(InData.CommandName);
+		CommandInfo.IsValid())
+	{
+		const IConsoleVariable* AsVariable = CommandInfo.Pin()->GetConsoleVariablePtr();
+		
+		if (AsVariable && AsVariable->TestFlags(ECVF_RenderThreadSafe))
+		{
+			UE_LOG(LogConsoleVariablesEditor, Warning,
+				TEXT("The console variable named %s is flagged as ECVF_RenderThreadSafe. The value on the render thread will lag behind the value on the main thread by one frame if r.OneFrameThreadLag is 1."),
+				*InData.CommandName
+			);
+		}
+	}
+	
 	RemoveConsoleVariable(InData.CommandName);
 	SavedCommands.Add(InData);
 }
 
-bool UConsoleVariablesAsset::RemoveConsoleVariable(const FString& InCommandString)
+bool UConsoleVariablesAsset::RemoveConsoleVariable(const FString InCommandString)
 {
 	FConsoleVariablesEditorAssetSaveData ExistingData;
 	if (FindSavedDataByCommandString(InCommandString, ExistingData))
diff --git a/Engine/Plugins/Experimental/ConsoleVariables/Source/ConsoleVariablesEditor/Private/ConsoleVariablesEditorCommandInfo.h b/Engine/Plugins/Experimental/ConsoleVariables/Source/ConsoleVariablesEditor/Private/ConsoleVariablesEditorCommandInfo.h
--- a/Engine/Plugins/Experimental/ConsoleVariables/Source/ConsoleVariablesEditor/Private/ConsoleVariablesEditorCommandInfo.h
+++ b/Engine/Plugins/Experimental/ConsoleVariables/Source/ConsoleVariablesEditor/Private/ConsoleVariablesEditorCommandInfo.h
@@ -213,6 +213,10 @@ struct FConsoleVariablesEditorCommandInfo
 			
 			return !AsVariable->GetString().Equals(InValueToCompare);
 		}
+		else if (ObjectType == EConsoleObjectType::NullObject || ObjectType == EConsoleObjectType::Command)
+		{
+			return true;
+		}
 
 		return false;
 	}
diff --git a/Engine/Plugins/Experimental/ConsoleVariables/Source/ConsoleVariablesEditor/Private/ConsoleVariablesEditorModule.cpp b/Engine/Plugins/Experimental/ConsoleVariables/Source/ConsoleVariablesEditor/Private/ConsoleVariablesEditorModule.cpp
--- a/Engine/Plugins/Experimental/ConsoleVariables/Source/ConsoleVariablesEditor/Private/ConsoleVariablesEditorModule.cpp
+++ b/Engine/Plugins/Experimental/ConsoleVariables/Source/ConsoleVariablesEditor/Private/ConsoleVariablesEditorModule.cpp
@@ -189,13 +189,30 @@ void FConsoleVariablesEditorModule::OnRemoteCvarChanged(const FString InName, co
 
 	if (GetDefault<UConcertCVarSynchronization>()->bSyncCVarTransactions)
 	{
+		bool bShouldExecute = true;
+		
 		if (const TWeakPtr<FConsoleVariablesEditorCommandInfo> CommandInfo =
 			FindCommandInfoByName(InName); CommandInfo.IsValid())
 		{
-			if (CommandInfo.Pin()->IsCurrentValueDifferentFromInputValue(InValue))
+			bShouldExecute = CommandInfo.Pin()->IsCurrentValueDifferentFromInputValue(InValue);
+		}
+
+		if (bShouldExecute)
+		{
+			GEngine->Exec(FConsoleVariablesEditorCommandInfo::GetCurrentWorld(),
+				*FString::Printf(TEXT("%s %s"), *InName, *InValue));
+			
+			EditingPresetAsset->AddOrSetConsoleObjectSavedData(
+				{
+					InName,
+					InValue,
+					ECheckBoxState::Checked
+				}
+			);
+
+			if (MainPanel->GetEditorListMode() == FConsoleVariablesEditorList::EConsoleVariablesEditorListMode::Preset)
 			{
-				GEngine->Exec(FConsoleVariablesEditorCommandInfo::GetCurrentWorld(),
-					*FString::Printf(TEXT("%s %s"), *InName, *InValue));
+				MainPanel->RebuildList();
 			}
 		}
 	}
diff --git a/Engine/Plugins/Experimental/ConsoleVariables/Source/ConsoleVariablesEditor/Private/Views/List/SConsoleVariablesEditorListRow.cpp b/Engine/Plugins/Experimental/ConsoleVariables/Source/ConsoleVariablesEditor/Private/Views/List/SConsoleVariablesEditorListRow.cpp
--- a/Engine/Plugins/Experimental/ConsoleVariables/Source/ConsoleVariablesEditor/Private/Views/List/SConsoleVariablesEditorListRow.cpp
+++ b/Engine/Plugins/Experimental/ConsoleVariables/Source/ConsoleVariablesEditor/Private/Views/List/SConsoleVariablesEditorListRow.cpp
@@ -538,28 +538,35 @@ TSharedRef<SWidget> SConsoleVariablesEditorListRow::GenerateValueCellWidget(
 			.ButtonStyle(&FAppStyle::Get().GetWidgetStyle<FButtonStyle>("NoBorder"))
 			.ContentPadding(0)
 			.Visibility_Lambda([this, PinnedItem]()
-			             {
-				             if (PinnedItem.IsValid() && PinnedItem->GetRowType() ==
-					             FConsoleVariablesEditorListRow::SingleCommand)
-				             {
-					             PinnedItem->SetDoesCurrentValueDifferFromPresetValue(
-						             PinnedItem->GetCommandInfo().Pin()->IsCurrentValueDifferentFromInputValue(
-							             PinnedItem->GetPresetValue()));
-
-					             return PinnedItem->IsRowChecked() && PinnedItem->
-					                    DoesCurrentValueDifferFromPresetValue()
-						                    ? EVisibility::Visible
-						                    : EVisibility::Collapsed;
-				             }
-
-				             return EVisibility::Collapsed;
-			             })
+             {
+	             if (PinnedItem.IsValid() && PinnedItem->GetRowType() ==
+		             FConsoleVariablesEditorListRow::SingleCommand)
+	             {
+		             const TSharedPtr<FConsoleVariablesEditorCommandInfo> CommandInfo =
+			             PinnedItem->GetCommandInfo().Pin();
+		            
+		            if (CommandInfo.IsValid() &&
+			            CommandInfo->ObjectType == FConsoleVariablesEditorCommandInfo::EConsoleObjectType::Variable)
+		            {
+			            PinnedItem->SetDoesCurrentValueDifferFromPresetValue(
+							 CommandInfo->IsCurrentValueDifferentFromInputValue(
+								 PinnedItem->GetPresetValue()));
+
+						 return PinnedItem->IsRowChecked() && PinnedItem->
+								DoesCurrentValueDifferFromPresetValue()
+									? EVisibility::Visible
+									: EVisibility::Collapsed;
+		            }
+	             }
+
+	             return EVisibility::Collapsed;
+             })
 			.OnClicked_Lambda([this, PinnedItem]()
-			             {
-				             PinnedItem->ResetToPresetValue();
+             {
+	             PinnedItem->ResetToPresetValue();
 
-				             return FReply::Handled();
-			             })
+	             return FReply::Handled();
+             })
 			[
 				SNew(SImage)
 				.Image(FAppStyle::Get().GetBrush("PropertyWindow.DiffersFromDefault"))
diff --git a/Engine/Plugins/Experimental/ConsoleVariables/Source/ConsoleVariablesEditor/Private/Views/List/SConsoleVariablesEditorListValueInput.cpp b/Engine/Plugins/Experimental/ConsoleVariables/Source/ConsoleVariablesEditor/Private/Views/List/SConsoleVariablesEditorListValueInput.cpp
--- a/Engine/Plugins/Experimental/ConsoleVariables/Source/ConsoleVariablesEditor/Private/Views/List/SConsoleVariablesEditorListValueInput.cpp
+++ b/Engine/Plugins/Experimental/ConsoleVariables/Source/ConsoleVariablesEditor/Private/Views/List/SConsoleVariablesEditorListValueInput.cpp
@@ -91,7 +91,7 @@ void SConsoleVariablesEditorListValueInput_Float::Construct(const FArguments& In
 			{
 				if (const IConsoleVariable* AsVariable = Item.Pin()->GetCommandInfo().Pin()->GetConsoleVariablePtr())
 				{
-					return FCString::Atof(*AsVariable->GetString());
+					return FCString::Atof(*FString::SanitizeFloat(AsVariable->GetFloat()));
 				}
 			}
 
@@ -100,9 +100,9 @@ void SConsoleVariablesEditorListValueInput_Float::Construct(const FArguments& In
 		.OnValueChanged_Lambda([this] (const float InValue)
 		{
 			check (Item.IsValid());
-			
+					
 			const FString ValueAsString = FString::SanitizeFloat(InValue);
-			
+					
 			if (const TSharedPtr<FConsoleVariablesEditorListRow> PinnedItem = Item.Pin();
 				!PinnedItem->GetCachedValue().Equals(ValueAsString))
 			{				
@@ -110,7 +110,7 @@ void SConsoleVariablesEditorListValueInput_Float::Construct(const FArguments& In
 
 				FConsoleVariablesEditorModule::Get().SendMultiUserConsoleVariableChange(
 					PinnedItem->GetCommandInfo().Pin()->Command, ValueAsString);
-			
+					
 				PinnedItem->SetCachedValue(ValueAsString);
 			}
 		})
@@ -450,6 +450,7 @@ void SConsoleVariablesEditorListValueInput_Command::Construct(const FArguments&
 				return FReply::Unhandled();
 			})
 			.IsEnabled(this, &SConsoleVariablesEditorListValueInput::IsRowChecked)
+			.ContentPadding(FMargin(0.f))
 			[
 				SNew(STextBlock)
 				.Justification(ETextJustify::Center)
diff --git a/Engine/Plugins/Experimental/ConsoleVariables/Source/ConsoleVariablesEditor/Private/Views/MainPanel/ConsoleVariablesEditorMainPanel.cpp b/Engine/Plugins/Experimental/ConsoleVariables/Source/ConsoleVariablesEditor/Private/Views/MainPanel/ConsoleVariablesEditorMainPanel.cpp
--- a/Engine/Plugins/Experimental/ConsoleVariables/Source/ConsoleVariablesEditor/Private/Views/MainPanel/ConsoleVariablesEditorMainPanel.cpp
+++ b/Engine/Plugins/Experimental/ConsoleVariables/Source/ConsoleVariablesEditor/Private/Views/MainPanel/ConsoleVariablesEditorMainPanel.cpp
@@ -54,14 +54,10 @@ TObjectPtr<UConsoleVariablesAsset> FConsoleVariablesEditorMainPanel::GetEditingA
 }
 
 void FConsoleVariablesEditorMainPanel::AddConsoleObjectToPreset(
-	const FString& InConsoleCommand, const FString& InValue, const bool bScrollToNewRow) const
+	const FString InConsoleCommand, const FString InValue, const bool bScrollToNewRow) const
 {
-	const TWeakObjectPtr<UConsoleVariablesAsset> EditingAsset = GetEditingAsset();
-	
-	if (EditingAsset.IsValid())
+	if (const TObjectPtr<UConsoleVariablesAsset> Asset = GetEditingAsset())
 	{
-		UConsoleVariablesAsset* Asset = EditingAsset.Get();
-
 		Asset->AddOrSetConsoleObjectSavedData(
 			{
 				InConsoleCommand,
@@ -77,7 +73,7 @@ void FConsoleVariablesEditorMainPanel::AddConsoleObjectToPreset(
 	}
 }
 
-void FConsoleVariablesEditorMainPanel::RebuildList(const FString& InConsoleCommandToScrollTo) const
+void FConsoleVariablesEditorMainPanel::RebuildList(const FString InConsoleCommandToScrollTo) const
 {
 	if (EditorList.IsValid())
 	{
diff --git a/Engine/Plugins/Experimental/ConsoleVariables/Source/ConsoleVariablesEditor/Private/Views/MainPanel/ConsoleVariablesEditorMainPanel.h b/Engine/Plugins/Experimental/ConsoleVariables/Source/ConsoleVariablesEditor/Private/Views/MainPanel/ConsoleVariablesEditorMainPanel.h
--- a/Engine/Plugins/Experimental/ConsoleVariables/Source/ConsoleVariablesEditor/Private/Views/MainPanel/ConsoleVariablesEditorMainPanel.h
+++ b/Engine/Plugins/Experimental/ConsoleVariables/Source/ConsoleVariablesEditor/Private/Views/MainPanel/ConsoleVariablesEditorMainPanel.h
@@ -23,11 +23,11 @@ public:
 	static TObjectPtr<UConsoleVariablesAsset> GetEditingAsset();
 
 	void AddConsoleObjectToPreset(
-		const FString& InConsoleCommand, const FString& InValue,
+		const FString InConsoleCommand, const FString InValue,
 		const bool bScrollToNewRow = false) const;
 
 	/** Regenerate the list items and refresh the list. Call when adding or removing variables. */
-	void RebuildList(const FString& InConsoleCommandToScrollTo = "") const;
+	void RebuildList(const FString InConsoleCommandToScrollTo = "") const;
 
 	/**
 	 * Refresh filters and sorting.
diff --git a/Engine/Plugins/Experimental/ConsoleVariables/Source/ConsoleVariablesEditor/Private/Views/MainPanel/SConsoleVariablesEditorCustomConsoleInputBox.cpp b/Engine/Plugins/Experimental/ConsoleVariables/Source/ConsoleVariablesEditor/Private/Views/MainPanel/SConsoleVariablesEditorCustomConsoleInputBox.cpp
--- a/Engine/Plugins/Experimental/ConsoleVariables/Source/ConsoleVariablesEditor/Private/Views/MainPanel/SConsoleVariablesEditorCustomConsoleInputBox.cpp
+++ b/Engine/Plugins/Experimental/ConsoleVariables/Source/ConsoleVariablesEditor/Private/Views/MainPanel/SConsoleVariablesEditorCustomConsoleInputBox.cpp
@@ -130,6 +130,14 @@ void SConsoleVariablesEditorCustomConsoleInputBox::Construct(
 	];
 }
 
+SConsoleVariablesEditorCustomConsoleInputBox::~SConsoleVariablesEditorCustomConsoleInputBox()
+{
+	MainPanelWidget.Reset();
+	InputText.Reset();
+	SuggestionBox.Reset();
+	SuggestionListView.Reset();
+}
+
 FReply SConsoleVariablesEditorCustomConsoleInputBox::OnPreviewKeyDown(const FGeometry& MyGeometry, const FKeyEvent& KeyEvent)
 {
 	if(SuggestionBox->IsOpen())
@@ -255,7 +263,7 @@ void SConsoleVariablesEditorCustomConsoleInputBox::OnTextChanged(const FText& In
 	}
 }
 
-FReply SConsoleVariablesEditorCustomConsoleInputBox::OnKeyCharHandler(const FGeometry& MyGeometry, const FCharacterEvent& InCharacterEvent)
+FReply SConsoleVariablesEditorCustomConsoleInputBox::OnKeyCharHandler(const FGeometry& MyGeometry, const FCharacterEvent& InCharacterEvent) const
 {
 	// Intercept tab if used for auto-complete
 	if (InCharacterEvent.GetCharacter() == '\t')
diff --git a/Engine/Plugins/Experimental/ConsoleVariables/Source/ConsoleVariablesEditor/Private/Views/MainPanel/SConsoleVariablesEditorCustomConsoleInputBox.h b/Engine/Plugins/Experimental/ConsoleVariables/Source/ConsoleVariablesEditor/Private/Views/MainPanel/SConsoleVariablesEditorCustomConsoleInputBox.h
--- a/Engine/Plugins/Experimental/ConsoleVariables/Source/ConsoleVariablesEditor/Private/Views/MainPanel/SConsoleVariablesEditorCustomConsoleInputBox.h
+++ b/Engine/Plugins/Experimental/ConsoleVariables/Source/ConsoleVariablesEditor/Private/Views/MainPanel/SConsoleVariablesEditorCustomConsoleInputBox.h
@@ -72,13 +72,15 @@ public:
 
 	void Construct(const FArguments& InArgs, TWeakPtr<SConsoleVariablesEditorMainPanel> InMainPanelWidget);
 
+	virtual ~SConsoleVariablesEditorCustomConsoleInputBox() override;
+
 	//~ Begin SWidget Interface
 	virtual FReply OnPreviewKeyDown(const FGeometry& MyGeometry, const FKeyEvent& KeyEvent) override;
 	//~ End SWidget Interface
 
 	void OnTextChanged(const FText& InText);
 
-	FReply OnKeyCharHandler(const FGeometry& MyGeometry, const FCharacterEvent& InCharacterEvent);
+	FReply OnKeyCharHandler(const FGeometry& MyGeometry, const FCharacterEvent& InCharacterEvent) const;
 
 	FReply HandleConsoleInputTextCommitted(const FGeometry& MyGeometry, const FKeyEvent& KeyPressed);
 	
diff --git a/Engine/Plugins/Experimental/ConsoleVariables/Source/ConsoleVariablesEditor/Private/Views/MainPanel/SConsoleVariablesEditorMainPanel.cpp b/Engine/Plugins/Experimental/ConsoleVariables/Source/ConsoleVariablesEditor/Private/Views/MainPanel/SConsoleVariablesEditorMainPanel.cpp
--- a/Engine/Plugins/Experimental/ConsoleVariables/Source/ConsoleVariablesEditor/Private/Views/MainPanel/SConsoleVariablesEditorMainPanel.cpp
+++ b/Engine/Plugins/Experimental/ConsoleVariables/Source/ConsoleVariablesEditor/Private/Views/MainPanel/SConsoleVariablesEditorMainPanel.cpp
@@ -36,24 +36,32 @@ void SConsoleVariablesEditorMainPanel::Construct(
 	
 	ChildSlot
 	[
-		SNew(SSplitter)
-		.Orientation(Orient_Vertical)
-		.ResizeMode(ESplitterResizeMode::FixedSize)
-		
-		+SSplitter::Slot().SizeRule(SSplitter::ESizeRule::SizeToContent)
+		SNew(SVerticalBox)
+			
+		+SVerticalBox::Slot()
+		.HAlign(HAlign_Fill)
+		.AutoHeight()
+		.Padding(FMargin(8.f, 0.f, 8.f, 0.f))
 		[
-	        GeneratePanelToolbar()
+			GeneratePanelToolbar()
 		]
 
-		+SSplitter::Slot()
+		+SVerticalBox::Slot()
 		[
-			MainPanel.Pin()->GetEditorList().Pin()->GetOrCreateWidget()
-		]
+			SNew(SSplitter)
+			.Orientation(Orient_Vertical)
+			.ResizeMode(ESplitterResizeMode::FixedSize)
 
-		+SSplitter::Slot()
-		[
-			SAssignNew(MultiUserDetailsBox, SVerticalBox)
-			.Visibility(EVisibility::Collapsed)
+			+SSplitter::Slot()
+			[
+				MainPanel.Pin()->GetEditorList().Pin()->GetOrCreateWidget()
+			]
+
+			+SSplitter::Slot()
+			[
+				SAssignNew(MultiUserDetailsBox, SVerticalBox)
+				.Visibility(EVisibility::Collapsed)
+			]
 		]
 	];
 
@@ -68,7 +76,7 @@ SConsoleVariablesEditorMainPanel::~SConsoleVariablesEditorMainPanel()
 	MultiUserDetailsBox.Reset();
 }
 
-FReply SConsoleVariablesEditorMainPanel::ValidateConsoleInput(const FText& CommittedText)
+FReply SConsoleVariablesEditorMainPanel::ValidateConsoleInput(const FText& CommittedText) const
 {
 	const FString CommandString = CommittedText.ToString().TrimStartAndEnd();
 	FString CommandKey;
@@ -122,7 +130,7 @@ FReply SConsoleVariablesEditorMainPanel::ValidateConsoleInput(const FText& Commi
 	return FReply::Handled();
 }
 
-void SConsoleVariablesEditorMainPanel::RefreshMultiUserDetails()
+void SConsoleVariablesEditorMainPanel::RefreshMultiUserDetails() const
 {
 	UConcertCVarSynchronization* CVarSync = GetMutableDefault<UConcertCVarSynchronization>();
 
@@ -142,7 +150,7 @@ void SConsoleVariablesEditorMainPanel::RefreshMultiUserDetails()
 	];
 }
 
-void SConsoleVariablesEditorMainPanel::ToggleMultiUserDetails(ECheckBoxState CheckState)
+void SConsoleVariablesEditorMainPanel::ToggleMultiUserDetails(ECheckBoxState CheckState) const
 {
 	const bool bShouldBeVisible = CheckState == ECheckBoxState::Checked;
 	MultiUserDetailsBox->SetVisibility(bShouldBeVisible ? EVisibility::SelfHitTestInvisible : EVisibility::Collapsed);
@@ -155,77 +163,73 @@ void SConsoleVariablesEditorMainPanel::ToggleMultiUserDetails(ECheckBoxState Che
 
 TSharedRef<SWidget> SConsoleVariablesEditorMainPanel::GeneratePanelToolbar()
 {
-	return SNew(SBorder)
-	        .Padding(0)
-	        .BorderImage(FAppStyle::Get().GetBrush("NoBorder"))
-			.HAlign(HAlign_Fill)
-	        [
-				SAssignNew(ToolbarHBox, SHorizontalBox)
+	return SAssignNew(ToolbarHBox, SHorizontalBox)
 				
-				// Add Console Variable input
-				+ SHorizontalBox::Slot()
-				.HAlign(HAlign_Fill)
-				.VAlign(VAlign_Fill)
-				.Padding(2.f, 2.f)
-				[
-					SNew(SConsoleVariablesEditorCustomConsoleInputBox, SharedThis(this))
-				]
+			// Add Console Variable input
+			+ SHorizontalBox::Slot()
+			.HAlign(HAlign_Fill)
+			.VAlign(VAlign_Fill)
+			[
+				SNew(SConsoleVariablesEditorCustomConsoleInputBox, SharedThis(this))
+			]
 
-				// Presets Management Button
-				+ SHorizontalBox::Slot()
-				.HAlign(HAlign_Right)
-				.VAlign(VAlign_Fill)
-				.AutoWidth()
+			// Presets Management Button
+			+ SHorizontalBox::Slot()
+			.HAlign(HAlign_Right)
+			.VAlign(VAlign_Fill)
+			.AutoWidth()
+			.Padding(FMargin(8.f, 0.f, 0.f, 0.f))
+			[
+				SNew(SComboButton)
+				.ToolTipText(LOCTEXT("PresetManagementButton_Tooltip", "Export the current CVar list to a preset, or import a copy of an existing preset."))
+				.ContentPadding(4.f)
+				.ComboButtonStyle(&FAppStyle::Get().GetWidgetStyle<FComboButtonStyle>("ComboButton"))
+				.OnGetMenuContent(this, &SConsoleVariablesEditorMainPanel::OnGeneratePresetsMenu)
+				.ForegroundColor(FStyleColors::Foreground)
+				.ButtonContent()
 				[
-					SNew(SComboButton)
-					.ToolTipText(LOCTEXT("PresetManagementButton_Tooltip", "Export the current CVar list to a preset, or import a copy of an existing preset."))
-					.ContentPadding(4.f)
-					.ComboButtonStyle(&FAppStyle::Get().GetWidgetStyle<FComboButtonStyle>("ComboButton"))
-					.OnGetMenuContent(this, &SConsoleVariablesEditorMainPanel::OnGeneratePresetsMenu)
-					.ForegroundColor(FStyleColors::Foreground)
-					.ButtonContent()
+					SNew(SHorizontalBox)
+
+					+ SHorizontalBox::Slot()
+					.Padding(0, 1, 4, 0)
+					.AutoWidth()
 					[
-						SNew(SHorizontalBox)
-
-						+ SHorizontalBox::Slot()
-						.Padding(0, 1, 4, 0)
-						.AutoWidth()
-						[
-							SNew(SImage)
-							.Image(FAppStyle::Get().GetBrush("AssetEditor.SaveAsset"))
-							.ColorAndOpacity(FSlateColor::UseForeground())
-						]
-
-						+ SHorizontalBox::Slot()
-						.Padding(0, 1, 0, 0)
-						[
-							SNew(STextBlock)
-							.Text(LOCTEXT("PresetsToolbarButton", "Presets"))
-						]
+						SNew(SImage)
+						.Image(FAppStyle::Get().GetBrush("AssetEditor.SaveAsset"))
+						.ColorAndOpacity(FSlateColor::UseForeground())
+					]
+
+					+ SHorizontalBox::Slot()
+					.Padding(0, 1, 0, 0)
+					[
+						SNew(STextBlock)
+						.Text(LOCTEXT("PresetsToolbarButton", "Presets"))
 					]
 				]
-        	];
+			];
 }
 
 void SConsoleVariablesEditorMainPanel::CreateConcertButtonIfNeeded()
 {
 	if (MainPanel.Pin()->GetMultiUserManager().IsInitialized())
 	{
+		constexpr int32 ButtonBoxSize = 28;
+		
 		// Toggle Multi-User Details
 		ToolbarHBox->AddSlot()
-		.HAlign(HAlign_Right)
 		.VAlign(VAlign_Fill)
 		.AutoWidth()
+		.Padding(FMargin(8.f, 1.f, 0.f, 1.f))
 		[
 			SNew(SBox)
-			.WidthOverride(28)
-			.HeightOverride(28)
+			.WidthOverride(ButtonBoxSize)
+			.HeightOverride(ButtonBoxSize)
 			[
 				SAssignNew(ConcertButtonPtr, SCheckBox)
-				.Padding(FMargin(4.f))
+				.Padding(4.f)
 				.ToolTipText(LOCTEXT("ShowConcertSettings_Tip", "Show the multi-user controls for Console Variables"))
 				.Style(&FAppStyle::Get().GetWidgetStyle<FCheckBoxStyle>("ToggleButtonCheckbox"))
-				.ForegroundColor(FStyleColors::Foreground)
+				.ForegroundColor(FSlateColor::UseForeground())
 				.IsChecked(false)
 				.OnCheckStateChanged_Raw(this, &SConsoleVariablesEditorMainPanel::ToggleMultiUserDetails)
 				[
diff --git a/Engine/Plugins/Experimental/ConsoleVariables/Source/ConsoleVariablesEditor/Private/Views/MainPanel/SConsoleVariablesEditorMainPanel.h b/Engine/Plugins/Experimental/ConsoleVariables/Source/ConsoleVariablesEditor/Private/Views/MainPanel/SConsoleVariablesEditorMainPanel.h
--- a/Engine/Plugins/Experimental/ConsoleVariables/Source/ConsoleVariablesEditor/Private/Views/MainPanel/SConsoleVariablesEditorMainPanel.h
+++ b/Engine/Plugins/Experimental/ConsoleVariables/Source/ConsoleVariablesEditor/Private/Views/MainPanel/SConsoleVariablesEditorMainPanel.h
@@ -26,10 +26,10 @@ public:
 	virtual ~SConsoleVariablesEditorMainPanel() override;
 
 	/** A handler for when text is committed to the Console Input Box that appears when clicking the "Add Console Variable" button */
-	FReply ValidateConsoleInput(const FText& CommittedText);
+	FReply ValidateConsoleInput(const FText& CommittedText) const;
 
-	void RefreshMultiUserDetails();
-	void ToggleMultiUserDetails(ECheckBoxState CheckState);
+	void RefreshMultiUserDetails() const;
+	void ToggleMultiUserDetails(ECheckBoxState CheckState) const;
 
 private:
 
diff --git a/Engine/Plugins/Experimental/ConsoleVariables/Source/ConsoleVariablesEditor/Public/ConsoleVariablesAsset.h b/Engine/Plugins/Experimental/ConsoleVariables/Source/ConsoleVariablesEditor/Public/ConsoleVariablesAsset.h
--- a/Engine/Plugins/Experimental/ConsoleVariables/Source/ConsoleVariablesEditor/Public/ConsoleVariablesAsset.h
+++ b/Engine/Plugins/Experimental/ConsoleVariables/Source/ConsoleVariablesEditor/Public/ConsoleVariablesAsset.h
@@ -63,13 +63,13 @@ public:
 	}
 
 	/** Outputs the FConsoleVariablesEditorCommandInfo matching InCommand. Returns whether a match was found. Case sensitive. */
-	bool FindSavedDataByCommandString(const FString& InCommandString, FConsoleVariablesEditorAssetSaveData& OutValue) const;
+	bool FindSavedDataByCommandString(const FString InCommandString, FConsoleVariablesEditorAssetSaveData& OutValue) const;
 
 	/** Set the value of a saved console variable if the name matches; add a new console variable to the list if a match is not found. */
 	void AddOrSetConsoleObjectSavedData(const FConsoleVariablesEditorAssetSaveData& InData);
 
 	/** Returns true if the element was found and successfully removed. */
-	bool RemoveConsoleVariable(const FString& InCommandString);
+	bool RemoveConsoleVariable(const FString InCommandString);
 
 	/** Copy data from input asset to this asset */
 	void CopyFrom(const UConsoleVariablesAsset* InAssetToCopy);
