Console Variables Editor:

-Fix for: Options box on CVAR tool flipping around

#[fyi] Jason.Walter

diff --git a/Engine/Plugins/Experimental/ConsoleVariables/Source/ConsoleVariablesEditor/Private/Views/MainPanel/SConsoleVariablesEditorCustomConsoleInputBox.cpp b/Engine/Plugins/Experimental/ConsoleVariables/Source/ConsoleVariablesEditor/Private/Views/MainPanel/SConsoleVariablesEditorCustomConsoleInputBox.cpp
--- a/Engine/Plugins/Experimental/ConsoleVariables/Source/ConsoleVariablesEditor/Private/Views/MainPanel/SConsoleVariablesEditorCustomConsoleInputBox.cpp
+++ b/Engine/Plugins/Experimental/ConsoleVariables/Source/ConsoleVariablesEditor/Private/Views/MainPanel/SConsoleVariablesEditorCustomConsoleInputBox.cpp
@@ -30,7 +30,7 @@ void SConsoleVariablesEditorCustomConsoleInputBox::Construct(
 	[
 		SAssignNew( SuggestionBox, SMenuAnchor )
 		.Method(GIsEditor ? EPopupMethod::CreateNewWindow : EPopupMethod::UseCurrentWindow)
-		.Placement( MenuPlacement_AboveAnchor )
+		.Placement( MenuPlacement_BelowAnchor )
 		[
 			SAssignNew(InputText, SEditableTextBox)
 			.Font(FEditorStyle::Get().GetWidgetStyle<FTextBlockStyle>("Log.Normal").Font)
