Console Variables Editor:

-Possible fix for: crash on nodes when performing various actions

#[fyi] Jason.Walter

diff --git a/Engine/Plugins/Experimental/ConsoleVariables/Source/ConsoleVariablesEditor/Private/ConsoleVariablesAsset.cpp b/Engine/Plugins/Experimental/ConsoleVariables/Source/ConsoleVariablesEditor/Private/ConsoleVariablesAsset.cpp
--- a/Engine/Plugins/Experimental/ConsoleVariables/Source/ConsoleVariablesEditor/Private/ConsoleVariablesAsset.cpp
+++ b/Engine/Plugins/Experimental/ConsoleVariables/Source/ConsoleVariablesEditor/Private/ConsoleVariablesAsset.cpp
@@ -17,8 +17,8 @@ void UConsoleVariablesAsset::ReplaceSavedCommands(const TArray<FConsoleVariables
 
 bool UConsoleVariablesAsset::FindSavedDataByCommandString(const FString InCommandString, FConsoleVariablesEditorAssetSaveData& OutValue) const
 {
-	if (const auto* Match = Algo::FindByPredicate(
-		SavedCommands,
+	if (const auto Match = 
+		SavedCommands.FindByPredicate(
 		[&InCommandString] (const FConsoleVariablesEditorAssetSaveData& Comparator)
 		{
 			return Comparator.CommandName.Equals(InCommandString);
