Console Variables Editor:

-Possible fix for: crash on nodes when performing various actions - Attempt 2

#[fyi] Jason.Walter

diff --git a/Engine/Plugins/Experimental/ConsoleVariables/Source/ConsoleVariablesEditor/Private/ConsoleVariablesAsset.cpp b/Engine/Plugins/Experimental/ConsoleVariables/Source/ConsoleVariablesEditor/Private/ConsoleVariablesAsset.cpp
--- a/Engine/Plugins/Experimental/ConsoleVariables/Source/ConsoleVariablesEditor/Private/ConsoleVariablesAsset.cpp
+++ b/Engine/Plugins/Experimental/ConsoleVariables/Source/ConsoleVariablesEditor/Private/ConsoleVariablesAsset.cpp
@@ -17,15 +17,13 @@ void UConsoleVariablesAsset::ReplaceSavedCommands(const TArray<FConsoleVariables
 
 bool UConsoleVariablesAsset::FindSavedDataByCommandString(const FString InCommandString, FConsoleVariablesEditorAssetSaveData& OutValue) const
 {
-	if (const auto Match = 
-		SavedCommands.FindByPredicate(
-		[&InCommandString] (const FConsoleVariablesEditorAssetSaveData& Comparator)
-		{
-			return Comparator.CommandName.Equals(InCommandString);
-		}))
+	for (const FConsoleVariablesEditorAssetSaveData& Command : SavedCommands)
 	{
-		OutValue = *Match;
-		return true;
+		if (Command.CommandName.Equals(InCommandString))
+		{
+			OutValue = Command;
+			return true;
+		}
 	}
 	
 	return false;
