Console Variables Editor:

-Replace "Enter Console Command" Form with "+ Console Variable" Button
-Change the search button from binoculars to saying "Search All". NB! An icon from Erik will replace the Binoculars when available

#[fyi] Jason.Walter

diff --git a/Engine/Plugins/Experimental/ConsoleVariables/Source/ConsoleVariablesEditor/ConsoleVariablesEditor.build.cs b/Engine/Plugins/Experimental/ConsoleVariables/Source/ConsoleVariablesEditor/ConsoleVariablesEditor.build.cs
--- a/Engine/Plugins/Experimental/ConsoleVariables/Source/ConsoleVariablesEditor/ConsoleVariablesEditor.build.cs
+++ b/Engine/Plugins/Experimental/ConsoleVariables/Source/ConsoleVariablesEditor/ConsoleVariablesEditor.build.cs
@@ -38,7 +38,8 @@ public class ConsoleVariablesEditor : ModuleRules
 				"PropertyEditor",
 				"Slate",
 				"SlateCore",
-				"ToolMenus",
+				"ToolMenus", 
+				"ToolWidgets",
 				"UnrealEd",
 				"WorkspaceMenuStructure"
 			}
diff --git a/Engine/Plugins/Experimental/ConsoleVariables/Source/ConsoleVariablesEditor/Private/Views/List/SConsoleVariablesEditorList.cpp b/Engine/Plugins/Experimental/ConsoleVariables/Source/ConsoleVariablesEditor/Private/Views/List/SConsoleVariablesEditorList.cpp
--- a/Engine/Plugins/Experimental/ConsoleVariables/Source/ConsoleVariablesEditor/Private/Views/List/SConsoleVariablesEditorList.cpp
+++ b/Engine/Plugins/Experimental/ConsoleVariables/Source/ConsoleVariablesEditor/Private/Views/List/SConsoleVariablesEditorList.cpp
@@ -58,6 +58,7 @@ void SConsoleVariablesEditorList::Construct(const FArguments& InArgs, TSharedRef
 			SNew(SHorizontalBox)
 
 			+SHorizontalBox::Slot()
+			.HAlign(HAlign_Fill)
 			.Padding(0.f, 1.f, 0.f, 1.f)
 			[
 				SAssignNew(ListSearchBoxPtr, SSearchBox)
@@ -65,13 +66,15 @@ void SConsoleVariablesEditorList::Construct(const FArguments& InArgs, TSharedRef
 				.OnTextChanged_Raw(this, &SConsoleVariablesEditorList::OnListViewSearchTextChanged)
 			]
 
-			// Global Search
+			// Global Search Button
 			+SHorizontalBox::Slot()
-			.Padding(8.f, 1.f, 0.f, 1.f)
 			.VAlign(VAlign_Center)
+			.HAlign(HAlign_Right)
 			.AutoWidth()
+			.Padding(FMargin(8.f, 0.f, 0.f, 0.f))
 			[
 				SNew(SButton)
+				.ContentPadding(4.f)
 				.VAlign(VAlign_Center)
 				.OnClicked_Lambda([this]()
 				{
@@ -81,9 +84,24 @@ void SConsoleVariablesEditorList::Construct(const FArguments& InArgs, TSharedRef
 				})
 				.ToolTipText(LOCTEXT("OpenInGlobalSearchButtonTooltip", "Search All Console Variables"))
 				[
-					SNew(STextBlock)
-					.TextStyle(FAppStyle::Get(), "FindResults.FindInBlueprints")
-					.Text(FText::FromString(FString(TEXT("\xf1e5"))) /*fa-binoculars*/)
+					SNew(SHorizontalBox)
+
+					+SHorizontalBox::Slot()
+					.Padding(0, 1, 4, 0)
+					.AutoWidth()
+					[
+						SNew(STextBlock)
+						.TextStyle(FAppStyle::Get(), "FindResults.FindInBlueprints")
+						.Text(FText::FromString(FString(TEXT("\xf1e5"))) /*fa-binoculars*/)
+					]
+
+					+SHorizontalBox::Slot()
+					.Padding(4.f, 0.f, 0.f, 0.f)
+					.AutoWidth()
+					[
+						SNew(STextBlock)
+						.Text(LOCTEXT("GlobalSearchButtonText","Search All"))
+					]
 				]
 			]
 
@@ -95,6 +113,7 @@ void SConsoleVariablesEditorList::Construct(const FArguments& InArgs, TSharedRef
 			.Padding(8.f, 1.f, 0.f, 1.f)
 			[
 				SAssignNew( ViewOptionsComboButton, SComboButton )
+				.ContentPadding(0)
 				.ToolTipText(LOCTEXT("ShowOptions_Tooltip", "Show options to affect the visibility of items in the Console Variables Editor list"))
 				.ComboButtonStyle( FAppStyle::Get(), "SimpleComboButtonWithIcon" ) // Use the tool bar item style for this button
 				.OnGetMenuContent( this, &SConsoleVariablesEditorList::BuildShowOptionsMenu)
diff --git a/Engine/Plugins/Experimental/ConsoleVariables/Source/ConsoleVariablesEditor/Private/Views/MainPanel/SConsoleVariablesEditorCustomConsoleInputBox.cpp b/Engine/Plugins/Experimental/ConsoleVariables/Source/ConsoleVariablesEditor/Private/Views/MainPanel/SConsoleVariablesEditorCustomConsoleInputBox.cpp
--- a/Engine/Plugins/Experimental/ConsoleVariables/Source/ConsoleVariablesEditor/Private/Views/MainPanel/SConsoleVariablesEditorCustomConsoleInputBox.cpp
+++ b/Engine/Plugins/Experimental/ConsoleVariables/Source/ConsoleVariablesEditor/Private/Views/MainPanel/SConsoleVariablesEditorCustomConsoleInputBox.cpp
@@ -203,6 +203,13 @@ FReply SConsoleVariablesEditorCustomConsoleInputBox::OnPreviewKeyDown(const FGeo
 	return FReply::Unhandled();
 }
 
+bool SConsoleVariablesEditorCustomConsoleInputBox::TakeKeyboardFocus() const
+{
+	check(InputText.IsValid());
+	
+	return FSlateApplication::Get().SetKeyboardFocus(InputText.ToSharedRef());
+}
+
 void SConsoleVariablesEditorCustomConsoleInputBox::OnTextChanged(const FText& InText)
 {
 	if(bIgnoreUIUpdate)
diff --git a/Engine/Plugins/Experimental/ConsoleVariables/Source/ConsoleVariablesEditor/Private/Views/MainPanel/SConsoleVariablesEditorCustomConsoleInputBox.h b/Engine/Plugins/Experimental/ConsoleVariables/Source/ConsoleVariablesEditor/Private/Views/MainPanel/SConsoleVariablesEditorCustomConsoleInputBox.h
--- a/Engine/Plugins/Experimental/ConsoleVariables/Source/ConsoleVariablesEditor/Private/Views/MainPanel/SConsoleVariablesEditorCustomConsoleInputBox.h
+++ b/Engine/Plugins/Experimental/ConsoleVariables/Source/ConsoleVariablesEditor/Private/Views/MainPanel/SConsoleVariablesEditorCustomConsoleInputBox.h
@@ -78,6 +78,8 @@ public:
 	virtual FReply OnPreviewKeyDown(const FGeometry& MyGeometry, const FKeyEvent& KeyEvent) override;
 	//~ End SWidget Interface
 
+	bool TakeKeyboardFocus() const;
+
 	void OnTextChanged(const FText& InText);
 
 	FReply OnKeyCharHandler(const FGeometry& MyGeometry, const FCharacterEvent& InCharacterEvent) const;
diff --git a/Engine/Plugins/Experimental/ConsoleVariables/Source/ConsoleVariablesEditor/Private/Views/MainPanel/SConsoleVariablesEditorMainPanel.cpp b/Engine/Plugins/Experimental/ConsoleVariables/Source/ConsoleVariablesEditor/Private/Views/MainPanel/SConsoleVariablesEditorMainPanel.cpp
--- a/Engine/Plugins/Experimental/ConsoleVariables/Source/ConsoleVariablesEditor/Private/Views/MainPanel/SConsoleVariablesEditorMainPanel.cpp
+++ b/Engine/Plugins/Experimental/ConsoleVariables/Source/ConsoleVariablesEditor/Private/Views/MainPanel/SConsoleVariablesEditorMainPanel.cpp
@@ -15,6 +15,7 @@
 #include "Framework/Application/SlateApplication.h"
 #include "Framework/MultiBox/MultiBoxBuilder.h"
 #include "IContentBrowserSingleton.h"
+#include "SPositiveActionButton.h"
 #include "Styling/AppStyle.h"
 #include "Styling/StyleColors.h"
 #include "Widgets/Images/SImage.h"
@@ -46,6 +47,15 @@ void SConsoleVariablesEditorMainPanel::Construct(
 			GeneratePanelToolbar()
 		]
 
+		+SVerticalBox::Slot()
+		.HAlign(HAlign_Fill)
+		.AutoHeight()
+		.Padding(FMargin(8.f, 0.f, 8.f, 0.f))
+		[
+			SAssignNew(AddConsoleObjectInputBoxPtr, SConsoleVariablesEditorCustomConsoleInputBox, SharedThis(this))
+			.Visibility(EVisibility::Collapsed)
+		]
+
 		+SVerticalBox::Slot()
 		[
 			SNew(SSplitter)
@@ -165,19 +175,20 @@ TSharedRef<SWidget> SConsoleVariablesEditorMainPanel::GeneratePanelToolbar()
 {
 	return SAssignNew(ToolbarHBox, SHorizontalBox)
 				
-			// Add Console Variable input
+			// Add Console Variable button
 			+ SHorizontalBox::Slot()
-			.HAlign(HAlign_Fill)
 			.VAlign(VAlign_Fill)
+			.HAlign(HAlign_Left)
 			[
-				SNew(SConsoleVariablesEditorCustomConsoleInputBox, SharedThis(this))
+				SNew(SPositiveActionButton)
+				.Text(LOCTEXT("ConsoleVariable", "Console Variable"))
+				.OnClicked(this, &SConsoleVariablesEditorMainPanel::HandleAddConsoleVariableButtonClicked)
 			]
 
 			// Presets Management Button
 			+ SHorizontalBox::Slot()
+			.VAlign(VAlign_Center)
 			.HAlign(HAlign_Right)
-			.VAlign(VAlign_Fill)
-			.AutoWidth()
 			.Padding(FMargin(8.f, 0.f, 0.f, 0.f))
 			[
 				SNew(SComboButton)
@@ -201,6 +212,7 @@ TSharedRef<SWidget> SConsoleVariablesEditorMainPanel::GeneratePanelToolbar()
 
 					+ SHorizontalBox::Slot()
 					.Padding(0, 1, 0, 0)
+					.AutoWidth()
 					[
 						SNew(STextBlock)
 						.Text(LOCTEXT("PresetsToolbarButton", "Presets"))
@@ -209,34 +221,42 @@ TSharedRef<SWidget> SConsoleVariablesEditorMainPanel::GeneratePanelToolbar()
 			];
 }
 
+FReply SConsoleVariablesEditorMainPanel::HandleAddConsoleVariableButtonClicked()
+{
+	if (AddConsoleObjectInputBoxPtr.IsValid())
+	{
+		AddConsoleObjectInputBoxPtr->SetVisibility(EVisibility::Visible);
+
+		AddConsoleObjectInputBoxPtr->TakeKeyboardFocus();
+
+		return FReply::Handled();
+	}
+
+	return FReply::Unhandled();
+}
+
 void SConsoleVariablesEditorMainPanel::CreateConcertButtonIfNeeded()
 {
 	if (MainPanel.Pin()->GetMultiUserManager().IsInitialized())
-	{
-		constexpr int32 ButtonBoxSize = 28;
-		
+	{		
 		// Toggle Multi-User Details
 		ToolbarHBox->AddSlot()
-		.VAlign(VAlign_Fill)
-		.AutoWidth()
-		.Padding(FMargin(8.f, 1.f, 0.f, 1.f))
+			.HAlign(HAlign_Right)
+			.VAlign(VAlign_Center)
+			.AutoWidth()
+			.Padding(11.f, 1.f, 0.f, 1.f)
 		[
-			SNew(SBox)
-			.WidthOverride(ButtonBoxSize)
-			.HeightOverride(ButtonBoxSize)
+			SAssignNew(ConcertButtonPtr, SCheckBox)
+			.Padding(0)
+			.ToolTipText(LOCTEXT("ShowConcertSettings_Tip", "Show the multi-user controls for Console Variables"))
+			.Style(&FAppStyle::Get().GetWidgetStyle<FCheckBoxStyle>("ToggleButtonCheckbox"))
+			.ForegroundColor(FSlateColor::UseForeground())
+			.IsChecked(false)
+			.OnCheckStateChanged_Raw(this, &SConsoleVariablesEditorMainPanel::ToggleMultiUserDetails)
 			[
-				SAssignNew(ConcertButtonPtr, SCheckBox)
-				.Padding(4.f)
-				.ToolTipText(LOCTEXT("ShowConcertSettings_Tip", "Show the multi-user controls for Console Variables"))
-				.Style(&FAppStyle::Get().GetWidgetStyle<FCheckBoxStyle>("ToggleButtonCheckbox"))
-				.ForegroundColor(FSlateColor::UseForeground())
-				.IsChecked(false)
-				.OnCheckStateChanged_Raw(this, &SConsoleVariablesEditorMainPanel::ToggleMultiUserDetails)
-				[
-					SNew(SImage)
-					.Image(FConsoleVariablesEditorStyle::Get().GetBrush("Concert.MultiUser"))
-					.ColorAndOpacity(FSlateColor::UseForeground())
-				]
+				SNew(SImage)
+				.Image(FConsoleVariablesEditorStyle::Get().GetBrush("Concert.MultiUser"))
+				.ColorAndOpacity(FSlateColor::UseForeground())
 			]
 		];
 	}
diff --git a/Engine/Plugins/Experimental/ConsoleVariables/Source/ConsoleVariablesEditor/Private/Views/MainPanel/SConsoleVariablesEditorMainPanel.h b/Engine/Plugins/Experimental/ConsoleVariables/Source/ConsoleVariablesEditor/Private/Views/MainPanel/SConsoleVariablesEditorMainPanel.h
--- a/Engine/Plugins/Experimental/ConsoleVariables/Source/ConsoleVariablesEditor/Private/Views/MainPanel/SConsoleVariablesEditorMainPanel.h
+++ b/Engine/Plugins/Experimental/ConsoleVariables/Source/ConsoleVariablesEditor/Private/Views/MainPanel/SConsoleVariablesEditorMainPanel.h
@@ -7,6 +7,7 @@
 
 class FConsoleVariablesEditorMainPanel;
 class SBox;
+class SConsoleVariablesEditorCustomConsoleInputBox;
 class SHorizontalBox;
 class SVerticalBox;
 class UConsoleVariablesAsset;
@@ -36,6 +37,8 @@ private:
 	/** A reference to the struct that controls this widget */
 	TWeakPtr<FConsoleVariablesEditorMainPanel> MainPanel;
 
+	TSharedPtr<SConsoleVariablesEditorCustomConsoleInputBox> AddConsoleObjectInputBoxPtr;
+
 	TSharedPtr<SHorizontalBox> ToolbarHBox;
 
 	/** A reference to the button which opens the plugin settings */
@@ -43,6 +46,8 @@ private:
 
 	/** Creates the toolbar at the top of the MainPanel widget */
 	TSharedRef<SWidget> GeneratePanelToolbar();
+
+	FReply HandleAddConsoleVariableButtonClicked();
 	
 	void CreateConcertButtonIfNeeded();
 
