Console Variables Editor:

-Search "pills" should all have the the same active/inactive color (currently they seem random).  Use the same blue as the checked state in the checkboxes.
-Change the exit search icon from a trash can to an X.
-Make sure this is right-aligned with the settings icon above it.
-Have the search form just say "Search..." like the Outliner (we can put the exposition in the tooltip).

-Fix for: CVars are reset when another cvar is added or removed from the list

#[fyi] Jason.Walter

diff --git a/Engine/Plugins/Experimental/ConsoleVariables/Source/ConsoleVariablesEditor/Private/Views/List/SConsoleVariablesEditorList.cpp b/Engine/Plugins/Experimental/ConsoleVariables/Source/ConsoleVariablesEditor/Private/Views/List/SConsoleVariablesEditorList.cpp
--- a/Engine/Plugins/Experimental/ConsoleVariables/Source/ConsoleVariablesEditor/Private/Views/List/SConsoleVariablesEditorList.cpp
+++ b/Engine/Plugins/Experimental/ConsoleVariables/Source/ConsoleVariablesEditor/Private/Views/List/SConsoleVariablesEditorList.cpp
@@ -62,7 +62,8 @@ void SConsoleVariablesEditorList::Construct(const FArguments& InArgs, TSharedRef
 			.Padding(0.f, 1.f, 0.f, 1.f)
 			[
 				SAssignNew(ListSearchBoxPtr, SSearchBox)
-				.HintText(LOCTEXT("ConsoleVariablesEditorList_SearchHintText", "Search tracked variables, values, sources or help text..."))
+				.HintText(LOCTEXT("ConsoleVariablesEditorList_SearchHintText", "Search..."))
+				.ToolTipText(LOCTEXT("ConsoleVariablesEditorList_TooltipText", "Search tracked variables, values, sources or help text"))
 				.OnTextChanged_Raw(this, &SConsoleVariablesEditorList::OnListViewSearchTextChanged)
 			]
 
@@ -129,6 +130,7 @@ void SConsoleVariablesEditorList::Construct(const FArguments& InArgs, TSharedRef
 
 		+SVerticalBox::Slot()
 		.AutoHeight()
+		.Padding(FMargin(8.f, 0.f, 8.f, 0.f))
 		[
 			SAssignNew(GlobalSearchesHBox, SHorizontalBox)
 			.Visibility_Lambda([this]()
@@ -147,18 +149,21 @@ void SConsoleVariablesEditorList::Construct(const FArguments& InArgs, TSharedRef
 			]
 
 			+SHorizontalBox::Slot()
-			.Padding(4.f, 0.f, 2.f, 0.f)
-			.VAlign(VAlign_Center)
 			.HAlign(HAlign_Right)
+			.VAlign(VAlign_Center)
 			.AutoWidth()
+			.Padding(8.f, 1.f, 2.f, 1.f)
 			[
 				// Remove Button
-				SAssignNew(RemoveGlobalSearchesButtonPtr, SButton)
+				SAssignNew(RemoveGlobalSearchesButtonPtr, SCheckBox)
+				.Padding(0)
 				.ToolTipText(
 					LOCTEXT("RemoveGlobalSearchesButtonTooltip",
 					"Remove all global searches from the console variables editor."))
-				.ButtonColorAndOpacity(FStyleColors::Transparent)
-				.OnClicked_Lambda([this]()
+				.Style(&FAppStyle::Get().GetWidgetStyle<FCheckBoxStyle>("ToggleButtonCheckbox"))
+				.ForegroundColor(FSlateColor::UseForeground())
+				.IsChecked(false)
+				.OnCheckStateChanged_Lambda([this] (ECheckBoxState NewCheckState)
 				{
 					ListModelPtr.Pin()->SetListMode(FConsoleVariablesEditorList::EConsoleVariablesEditorListMode::Preset);
 
@@ -166,13 +171,13 @@ void SConsoleVariablesEditorList::Construct(const FArguments& InArgs, TSharedRef
 					CurrentGlobalSearches.Empty();
 
 					RebuildList();
-					
-					return FReply::Handled();
+
+					RemoveGlobalSearchesButtonPtr->SetIsChecked(false);
 				})
 				[
 					SNew(SImage)
 					.Visibility(EVisibility::SelfHitTestInvisible)
-					.Image(FAppStyle::Get().GetBrush("Icons.Delete"))
+					.Image(FAppStyle::Get().GetBrush("Icons.X"))
 					.ColorAndOpacity(FSlateColor::UseForeground())
 				]
 			]
@@ -444,8 +449,12 @@ void SConsoleVariablesEditorList::SetTreeViewItems(const TArray<FConsoleVariable
 void SConsoleVariablesEditorList::UpdatePresetValuesForSave(const TObjectPtr<UConsoleVariablesAsset> InAsset) const
 {
 	TArray<FConsoleVariablesEditorAssetSaveData> NewSavedCommands;
+
+	const TArray<FConsoleVariablesEditorListRowPtr>& Items =
+		ListModelPtr.IsValid() && ListModelPtr.Pin()->GetListMode() ==
+			FConsoleVariablesEditorList::EConsoleVariablesEditorListMode::Preset ? TreeViewRootObjects : LastPresetObjects;
 	
-	for (const FConsoleVariablesEditorListRowPtr& Item : TreeViewRootObjects)
+	for (const FConsoleVariablesEditorListRowPtr& Item : Items)
 	{
 		if (const TWeakPtr<FConsoleVariablesEditorCommandInfo> CommandInfo = Item->GetCommandInfo(); CommandInfo.IsValid())
 		{
diff --git a/Engine/Plugins/Experimental/ConsoleVariables/Source/ConsoleVariablesEditor/Private/Views/List/SConsoleVariablesEditorList.h b/Engine/Plugins/Experimental/ConsoleVariables/Source/ConsoleVariablesEditor/Private/Views/List/SConsoleVariablesEditorList.h
--- a/Engine/Plugins/Experimental/ConsoleVariables/Source/ConsoleVariablesEditor/Private/Views/List/SConsoleVariablesEditorList.h
+++ b/Engine/Plugins/Experimental/ConsoleVariables/Source/ConsoleVariablesEditor/Private/Views/List/SConsoleVariablesEditorList.h
@@ -140,7 +140,7 @@ private:
 	TSharedPtr<SWrapBox> GlobalSearchesContainer;
 	TArray<TSharedRef<SConsoleVariablesEditorGlobalSearchToggle>> CurrentGlobalSearches;
 	
-	TSharedPtr<SButton> RemoveGlobalSearchesButtonPtr;
+	TSharedPtr<SCheckBox> RemoveGlobalSearchesButtonPtr;
 	TSharedPtr<SBox> ListBoxContainerPtr;
 
 	//  Tree View Implementation
diff --git a/Engine/Plugins/Experimental/ConsoleVariables/Source/ConsoleVariablesEditor/Private/Views/MainPanel/ConsoleVariablesEditorMainPanel.cpp b/Engine/Plugins/Experimental/ConsoleVariables/Source/ConsoleVariablesEditor/Private/Views/MainPanel/ConsoleVariablesEditorMainPanel.cpp
--- a/Engine/Plugins/Experimental/ConsoleVariables/Source/ConsoleVariablesEditor/Private/Views/MainPanel/ConsoleVariablesEditorMainPanel.cpp
+++ b/Engine/Plugins/Experimental/ConsoleVariables/Source/ConsoleVariablesEditor/Private/Views/MainPanel/ConsoleVariablesEditorMainPanel.cpp
@@ -58,6 +58,8 @@ void FConsoleVariablesEditorMainPanel::AddConsoleObjectToPreset(
 {
 	if (const TObjectPtr<UConsoleVariablesAsset> Asset = GetEditingAsset())
 	{
+		UpdatePresetValuesForSave(Asset);
+		
 		Asset->AddOrSetConsoleObjectSavedData(
 			{
 				InConsoleCommand,
diff --git a/Engine/Plugins/Experimental/ConsoleVariables/Source/ConsoleVariablesEditor/Private/Views/MainPanel/SConsoleVariablesEditorMainPanel.cpp b/Engine/Plugins/Experimental/ConsoleVariables/Source/ConsoleVariablesEditor/Private/Views/MainPanel/SConsoleVariablesEditorMainPanel.cpp
--- a/Engine/Plugins/Experimental/ConsoleVariables/Source/ConsoleVariablesEditor/Private/Views/MainPanel/SConsoleVariablesEditorMainPanel.cpp
+++ b/Engine/Plugins/Experimental/ConsoleVariables/Source/ConsoleVariablesEditor/Private/Views/MainPanel/SConsoleVariablesEditorMainPanel.cpp
@@ -193,7 +193,7 @@ TSharedRef<SWidget> SConsoleVariablesEditorMainPanel::GeneratePanelToolbar()
 			[
 				SNew(SComboButton)
 				.ToolTipText(LOCTEXT("PresetManagementButton_Tooltip", "Export the current CVar list to a preset, or import a copy of an existing preset."))
-				.ContentPadding(4.f)
+				.ContentPadding(4)
 				.ComboButtonStyle(&FAppStyle::Get().GetWidgetStyle<FComboButtonStyle>("ComboButton"))
 				.OnGetMenuContent(this, &SConsoleVariablesEditorMainPanel::OnGeneratePresetsMenu)
 				.ForegroundColor(FStyleColors::Foreground)
@@ -241,10 +241,10 @@ void SConsoleVariablesEditorMainPanel::CreateConcertButtonIfNeeded()
 	{		
 		// Toggle Multi-User Details
 		ToolbarHBox->AddSlot()
-			.HAlign(HAlign_Right)
-			.VAlign(VAlign_Center)
-			.AutoWidth()
-			.Padding(11.f, 1.f, 0.f, 1.f)
+		.HAlign(HAlign_Right)
+		.VAlign(VAlign_Center)
+		.AutoWidth()
+		.Padding(11.f, 1.f, 1.f, 1.f)
 		[
 			SAssignNew(ConcertButtonPtr, SCheckBox)
 			.Padding(0)
diff --git a/Engine/Plugins/Experimental/ConsoleVariables/Source/ConsoleVariablesEditor/Private/Views/Widgets/SConsoleVariablesEditorGlobalSearchToggle.cpp b/Engine/Plugins/Experimental/ConsoleVariables/Source/ConsoleVariablesEditor/Private/Views/Widgets/SConsoleVariablesEditorGlobalSearchToggle.cpp
--- a/Engine/Plugins/Experimental/ConsoleVariables/Source/ConsoleVariablesEditor/Private/Views/Widgets/SConsoleVariablesEditorGlobalSearchToggle.cpp
+++ b/Engine/Plugins/Experimental/ConsoleVariables/Source/ConsoleVariablesEditor/Private/Views/Widgets/SConsoleVariablesEditorGlobalSearchToggle.cpp
@@ -4,6 +4,7 @@
 
 #include "Framework/Application/SlateApplication.h"
 #include "Styling/AppStyle.h"
+#include "Styling/StyleColors.h"
 #include "Widgets/Images/SImage.h"
 #include "Widgets/Layout/SBorder.h"
 #include "Widgets/SBoxPanel.h"
@@ -19,8 +20,9 @@ void SConsoleVariablesEditorGlobalSearchToggle::Construct(const FArguments& Args
 	OnToggleRightButtonClicked = Args._OnToggleRightButtonClicked;
 	OnToggleClickedOnce = Args._OnToggleClickedOnce;
 
-	CheckedColor = FColor::MakeRandomColor();
 	GlobalSearchText = InButtonText;
+	CheckedColor = USlateThemeManager::Get().GetColor(EStyleColor::Primary);
+	UncheckedColor = USlateThemeManager::Get().GetColor(EStyleColor::Black);
 	
 	ChildSlot
     [
diff --git a/Engine/Plugins/Experimental/ConsoleVariables/Source/ConsoleVariablesEditor/Private/Views/Widgets/SConsoleVariablesEditorGlobalSearchToggle.h b/Engine/Plugins/Experimental/ConsoleVariables/Source/ConsoleVariablesEditor/Private/Views/Widgets/SConsoleVariablesEditorGlobalSearchToggle.h
--- a/Engine/Plugins/Experimental/ConsoleVariables/Source/ConsoleVariablesEditor/Private/Views/Widgets/SConsoleVariablesEditorGlobalSearchToggle.h
+++ b/Engine/Plugins/Experimental/ConsoleVariables/Source/ConsoleVariablesEditor/Private/Views/Widgets/SConsoleVariablesEditorGlobalSearchToggle.h
@@ -57,8 +57,8 @@ private:
 	 * The responsibility is on external classes to remove the button from the UI the external class controls. */
 	bool bIsMarkedForDelete = false;
 
-	FColor CheckedColor;
-	FColor UncheckedColor = FColor::Black;
+	FLinearColor CheckedColor;
+	FLinearColor UncheckedColor;
 
 	FText GlobalSearchText;
 };
