Console Variables Editor:

-Pad top/bottom of the '+Console Variable' button using the content browser as spacing reference
-Collapse Console Input Field when focus is taken from it
-Various Padding and placement improvements

#[fyi] Jason.Walter

diff --git a/Engine/Plugins/Experimental/ConsoleVariables/Source/ConsoleVariablesEditor/Private/Views/MainPanel/SConsoleVariablesEditorCustomConsoleInputBox.cpp b/Engine/Plugins/Experimental/ConsoleVariables/Source/ConsoleVariablesEditor/Private/Views/MainPanel/SConsoleVariablesEditorCustomConsoleInputBox.cpp
--- a/Engine/Plugins/Experimental/ConsoleVariables/Source/ConsoleVariablesEditor/Private/Views/MainPanel/SConsoleVariablesEditorCustomConsoleInputBox.cpp
+++ b/Engine/Plugins/Experimental/ConsoleVariables/Source/ConsoleVariablesEditor/Private/Views/MainPanel/SConsoleVariablesEditorCustomConsoleInputBox.cpp
@@ -37,7 +37,16 @@ void SConsoleVariablesEditorCustomConsoleInputBox::Construct(
 			.HintText(LOCTEXT("ConsoleCommandExecutorHintText", "Enter Console Command"))
 			.OnTextChanged(this, &SConsoleVariablesEditorCustomConsoleInputBox::OnTextChanged)
 			.OnKeyCharHandler(this, &SConsoleVariablesEditorCustomConsoleInputBox::OnKeyCharHandler)
-			.OnKeyDownHandler(this, &SConsoleVariablesEditorCustomConsoleInputBox::HandleConsoleInputTextCommitted)
+			.OnKeyDownHandler(this, &SConsoleVariablesEditorCustomConsoleInputBox::OnKeyDownHandler)
+			.OnTextCommitted_Lambda([this] (const FText& InText, const ETextCommit::Type CommitType)
+			{
+				// Hide widget when focus is lost
+				if (CommitType == ETextCommit::OnUserMovedFocus)
+				{
+					SuggestionBox->SetIsOpen(false);
+					SetVisibility(EVisibility::Collapsed);
+				}
+			})
 			.ClearKeyboardFocusOnCommit(false)
 		
 		]
@@ -287,7 +296,7 @@ FReply SConsoleVariablesEditorCustomConsoleInputBox::OnKeyCharHandler(const FGeo
 	return FReply::Unhandled();
 }
 
-FReply SConsoleVariablesEditorCustomConsoleInputBox::HandleConsoleInputTextCommitted(const FGeometry& MyGeometry, const FKeyEvent& KeyPressed)
+FReply SConsoleVariablesEditorCustomConsoleInputBox::OnKeyDownHandler(const FGeometry& MyGeometry, const FKeyEvent& KeyPressed)
 {
 	if (KeyPressed.GetKey().GetFName() == TEXT("Enter"))
 	{
diff --git a/Engine/Plugins/Experimental/ConsoleVariables/Source/ConsoleVariablesEditor/Private/Views/MainPanel/SConsoleVariablesEditorCustomConsoleInputBox.h b/Engine/Plugins/Experimental/ConsoleVariables/Source/ConsoleVariablesEditor/Private/Views/MainPanel/SConsoleVariablesEditorCustomConsoleInputBox.h
--- a/Engine/Plugins/Experimental/ConsoleVariables/Source/ConsoleVariablesEditor/Private/Views/MainPanel/SConsoleVariablesEditorCustomConsoleInputBox.h
+++ b/Engine/Plugins/Experimental/ConsoleVariables/Source/ConsoleVariablesEditor/Private/Views/MainPanel/SConsoleVariablesEditorCustomConsoleInputBox.h
@@ -84,7 +84,7 @@ public:
 
 	FReply OnKeyCharHandler(const FGeometry& MyGeometry, const FCharacterEvent& InCharacterEvent) const;
 
-	FReply HandleConsoleInputTextCommitted(const FGeometry& MyGeometry, const FKeyEvent& KeyPressed);
+	FReply OnKeyDownHandler(const FGeometry& MyGeometry, const FKeyEvent& KeyPressed);
 	
 	void SetSuggestions(TArray<FString>& Elements, FText Highlight);
 
diff --git a/Engine/Plugins/Experimental/ConsoleVariables/Source/ConsoleVariablesEditor/Private/Views/MainPanel/SConsoleVariablesEditorMainPanel.cpp b/Engine/Plugins/Experimental/ConsoleVariables/Source/ConsoleVariablesEditor/Private/Views/MainPanel/SConsoleVariablesEditorMainPanel.cpp
--- a/Engine/Plugins/Experimental/ConsoleVariables/Source/ConsoleVariablesEditor/Private/Views/MainPanel/SConsoleVariablesEditorMainPanel.cpp
+++ b/Engine/Plugins/Experimental/ConsoleVariables/Source/ConsoleVariablesEditor/Private/Views/MainPanel/SConsoleVariablesEditorMainPanel.cpp
@@ -179,6 +179,7 @@ TSharedRef<SWidget> SConsoleVariablesEditorMainPanel::GeneratePanelToolbar()
 			+ SHorizontalBox::Slot()
 			.VAlign(VAlign_Fill)
 			.HAlign(HAlign_Left)
+			.Padding(FMargin(0, 4))
 			[
 				SNew(SPositiveActionButton)
 				.Text(LOCTEXT("ConsoleVariable", "Console Variable"))
@@ -189,11 +190,11 @@ TSharedRef<SWidget> SConsoleVariablesEditorMainPanel::GeneratePanelToolbar()
 			+ SHorizontalBox::Slot()
 			.VAlign(VAlign_Center)
 			.HAlign(HAlign_Right)
-			.Padding(FMargin(8.f, 0.f, 0.f, 0.f))
+			.Padding(FMargin(8.f, 4, 0.f, 4))
 			[
 				SNew(SComboButton)
 				.ToolTipText(LOCTEXT("PresetManagementButton_Tooltip", "Export the current CVar list to a preset, or import a copy of an existing preset."))
-				.ContentPadding(4)
+				.ContentPadding(FMargin(4, 0.5f))
 				.ComboButtonStyle(&FAppStyle::Get().GetWidgetStyle<FComboButtonStyle>("ComboButton"))
 				.OnGetMenuContent(this, &SConsoleVariablesEditorMainPanel::OnGeneratePresetsMenu)
 				.ForegroundColor(FStyleColors::Foreground)
