Console Variables Editor:

-Support Global Search AND vs. OR Behavior
-Local Search refactor to reflect global search behaviour
-Message is now displayed in UI when global search fails
-Added Blueprint Support Library
-Add thumbnails
-Replace Global Search button icon
-Add outline icon for global search items that are not yet favorited
-Spacing improvements
-Set Action Button style to "NoBorder"

diff --git a/Engine/Plugins/Experimental/ConsoleVariables/Source/ConsoleVariablesEditor/Private/ConsoleVariablesEditorFunctionLibrary.cpp b/Engine/Plugins/Experimental/ConsoleVariables/Source/ConsoleVariablesEditor/Private/ConsoleVariablesEditorFunctionLibrary.cpp
new file mode 100644
--- /dev/null
+++ b/Engine/Plugins/Experimental/ConsoleVariables/Source/ConsoleVariablesEditor/Private/ConsoleVariablesEditorFunctionLibrary.cpp
@@ -0,0 +1,218 @@
+// Copyright Epic Games, Inc. All Rights Reserved.
+
+#include "ConsoleVariablesEditorFunctionLibrary.h"
+
+#include "ConsoleVariablesAsset.h"
+#include "ConsoleVariablesEditorLog.h"
+#include "ConsoleVariablesEditorModule.h"
+#include "MultiUser/ConsoleVariableSyncData.h"
+#include "Widgets/Input/SCheckBox.h"
+
+UConsoleVariablesAsset* UConsoleVariablesEditorFunctionLibrary::GetCurrentlyLoadedPreset()
+{
+	const FConsoleVariablesEditorModule& ConsoleVariablesEditorModule = FConsoleVariablesEditorModule::Get();
+	
+	return ConsoleVariablesEditorModule.GetPresetAsset();
+}
+
+void UConsoleVariablesEditorFunctionLibrary::LoadPresetIntoConsoleVariablesEditor(const UConsoleVariablesAsset* InAsset)
+{
+	if (!InAsset)
+	{
+		UE_LOG(LogConsoleVariablesEditor, Error,
+			TEXT("%hs: InAsset was nullptr, please verify that InAsset is valid."),
+			__FUNCTION__);
+		return;
+	}
+	
+	const FConsoleVariablesEditorModule& ConsoleVariablesEditorModule = FConsoleVariablesEditorModule::Get();
+
+	ConsoleVariablesEditorModule.OpenConsoleVariablesDialogWithAssetSelected(InAsset);
+}
+
+bool UConsoleVariablesEditorFunctionLibrary::CopyCurrentListToAsset(UConsoleVariablesAsset* InAsset)
+{
+	if (!InAsset)
+	{
+		UE_LOG(LogConsoleVariablesEditor, Error,
+			TEXT("%hs: InAsset was nullptr, please verify that InAsset is valid."),
+			__FUNCTION__);
+		return false;
+	}
+	
+	const FConsoleVariablesEditorModule& ConsoleVariablesEditorModule = FConsoleVariablesEditorModule::Get();
+
+	InAsset->CopyFrom(ConsoleVariablesEditorModule.GetPresetAsset());
+	InAsset->MarkPackageDirty();
+	return true;
+}
+
+bool UConsoleVariablesEditorFunctionLibrary::AddValidatedCommandToCurrentPreset(const FString NewCommand)
+{
+	const FConsoleVariablesEditorModule& ConsoleVariablesEditorModule = FConsoleVariablesEditorModule::Get();
+
+	if (const FReply Reply = ConsoleVariablesEditorModule.ValidateConsoleInputAndAddToCurrentPreset(FText::FromString(NewCommand));
+		Reply.IsEventHandled())
+	{
+		GetCurrentlyLoadedPreset()->MarkPackageDirty();
+
+		return true;
+	}
+
+	return false;
+}
+
+bool UConsoleVariablesEditorFunctionLibrary::RemoveCommandFromCurrentPreset(const FString NewCommand)
+{
+	FConsoleVariablesEditorModule& ConsoleVariablesEditorModule = FConsoleVariablesEditorModule::Get();
+	
+	ConsoleVariablesEditorModule.UpdatePresetValuesForSave(GetCurrentlyLoadedPreset());
+	
+	if (const bool bSuccess = GetCurrentlyLoadedPreset()->RemoveConsoleVariable(NewCommand))
+	{
+		ConsoleVariablesEditorModule.RebuildList();
+		
+		GetCurrentlyLoadedPreset()->MarkPackageDirty();
+
+		return bSuccess;
+	}
+
+	return false;
+}
+
+bool UConsoleVariablesEditorFunctionLibrary::GetListOfCommandsFromPreset(
+ 	const UConsoleVariablesAsset* InAsset, TArray<FString>& OutCommandList)
+{
+	if (!InAsset)
+	{
+		UE_LOG(LogConsoleVariablesEditor, Error,
+			TEXT("%hs: InAsset was nullptr, please verify that InAsset is valid."),
+			__FUNCTION__);
+		return false;
+	}
+
+	OutCommandList.Empty();
+	for (const FConsoleVariablesEditorAssetSaveData& Command : InAsset->GetSavedCommands())
+	{
+		OutCommandList.Add(Command.CommandName);
+	}
+
+	if (OutCommandList.Num() < 0)
+	{
+		UE_LOG(LogConsoleVariablesEditor, Error,
+			TEXT("%hs: InAsset has no saved commands. Please verify you're using the right asset."),
+			__FUNCTION__);
+		return false;
+	}
+	
+	return true;
+}
+
+bool UConsoleVariablesEditorFunctionLibrary::SetConsoleVariableByName_Float(const FString InCommandName,
+	const float InValue)
+{
+	return SetConsoleVariableByName_String(InCommandName, FString::SanitizeFloat(InValue));
+}
+
+bool UConsoleVariablesEditorFunctionLibrary::SetConsoleVariableByName_Bool(const FString InCommandName,
+	const bool InValue)
+{
+	return SetConsoleVariableByName_String(InCommandName, InValue ? "true" : "false");
+}
+
+bool UConsoleVariablesEditorFunctionLibrary::SetConsoleVariableByName_Int(const FString InCommandName,
+	const int32 InValue)
+{
+	return SetConsoleVariableByName_String(InCommandName, FString::FromInt(InValue));
+}
+
+bool UConsoleVariablesEditorFunctionLibrary::SetConsoleVariableByName_String(const FString InCommandName,
+	const FString InValue)
+{
+	FConsoleVariablesEditorModule& ConsoleVariablesEditorModule = FConsoleVariablesEditorModule::Get();
+
+	if (const TWeakPtr<FConsoleVariablesEditorCommandInfo> CommandInfo =
+		ConsoleVariablesEditorModule.FindCommandInfoByName(InCommandName); CommandInfo.IsValid())
+	{
+		if (CommandInfo.Pin()->GetConsoleVariablePtr())
+		{
+			CommandInfo.Pin()->ExecuteCommand(InValue);
+			return true;
+		}
+		
+		UE_LOG(LogConsoleVariablesEditor, Error,
+			TEXT("%hs: Command %s is not a variable type. Please do not enter console commands, only console variables."),
+			__FUNCTION__, *InCommandName);
+	}
+	else
+	{
+		UE_LOG(LogConsoleVariablesEditor, Error,
+			TEXT("%hs: FConsoleVariablesEditorCommandInfo was not foundwith given name: %s."),
+			__FUNCTION__, *InCommandName);
+	}
+	return false;
+}
+
+bool UConsoleVariablesEditorFunctionLibrary::GetConsoleVariableStringValue(
+	const FString InCommandName, FString& OutValue)
+{
+	FConsoleVariablesEditorModule& ConsoleVariablesEditorModule = FConsoleVariablesEditorModule::Get();
+
+	if (const TWeakPtr<FConsoleVariablesEditorCommandInfo> CommandInfo =
+		ConsoleVariablesEditorModule.FindCommandInfoByName(InCommandName); CommandInfo.IsValid())
+	{
+		if (const IConsoleVariable* AsVariable = CommandInfo.Pin()->GetConsoleVariablePtr())
+		{
+			OutValue = AsVariable->GetString();
+			return true;
+		}
+	}
+
+	UE_LOG(LogConsoleVariablesEditor, Error,
+		TEXT("%hs: FConsoleVariablesEditorCommandInfo was not foundwith given name: %s."),
+		__FUNCTION__, *InCommandName);
+	return false;
+}
+
+bool UConsoleVariablesEditorFunctionLibrary::GetConsoleVariableSourceByName(const FString InCommandName, FString& OutValue)
+{
+	FConsoleVariablesEditorModule& ConsoleVariablesEditorModule = FConsoleVariablesEditorModule::Get();
+
+	if (const TWeakPtr<FConsoleVariablesEditorCommandInfo> CommandInfo =
+		ConsoleVariablesEditorModule.FindCommandInfoByName(InCommandName); CommandInfo.IsValid())
+	{
+		OutValue = CommandInfo.Pin()->GetSourceAsText().ToString();
+		return true;
+	}
+
+	UE_LOG(LogConsoleVariablesEditor, Error,
+		TEXT("%hs: FConsoleVariablesEditorCommandInfo was not found with given name: %s."),
+		__FUNCTION__, *InCommandName);
+	return false;
+}
+
+bool UConsoleVariablesEditorFunctionLibrary::GetEnableMultiUserCVarSync()
+{
+	if (const UConcertCVarSynchronization* CVarSync = GetDefault<UConcertCVarSynchronization>())
+	{
+		return CVarSync->bSyncCVarTransactions;
+	}
+
+	UE_LOG(LogConsoleVariablesEditor, Error,
+		TEXT("%hs: UConcertCVarSynchronization was not found, please verify the Multi-user plugin is enabled and this object class is compiled."),
+		__FUNCTION__);
+	return false;
+}
+
+void UConsoleVariablesEditorFunctionLibrary::SetEnableMultiUserCVarSync(bool bNewSetting)
+{
+	if (UConcertCVarSynchronization* CVarSync = GetMutableDefault<UConcertCVarSynchronization>())
+	{
+		CVarSync->bSyncCVarTransactions = bNewSetting;
+		return;
+	}
+
+	UE_LOG(LogConsoleVariablesEditor, Error,
+		TEXT("%hs: UConcertCVarSynchronization was not found, please verify the Multi-user plugin is enabled and this object class is compiled."),
+		__FUNCTION__);
+}
diff --git a/Engine/Plugins/Experimental/ConsoleVariables/Source/ConsoleVariablesEditor/Private/ConsoleVariablesEditorModule.cpp b/Engine/Plugins/Experimental/ConsoleVariables/Source/ConsoleVariablesEditor/Private/ConsoleVariablesEditorModule.cpp
--- a/Engine/Plugins/Experimental/ConsoleVariables/Source/ConsoleVariablesEditor/Private/ConsoleVariablesEditorModule.cpp
+++ b/Engine/Plugins/Experimental/ConsoleVariables/Source/ConsoleVariablesEditor/Private/ConsoleVariablesEditorModule.cpp
@@ -9,7 +9,7 @@
 #include "MultiUser/ConsoleVariableSyncData.h"
 #include "Views/MainPanel/ConsoleVariablesEditorMainPanel.h"
 
-#include "Algo/AnyOf.h"
+#include "Algo/AllOf.h"
 #include "Algo/Find.h"
 #include "Framework/Docking/TabManager.h"
 #include "ISettingsModule.h"
@@ -57,6 +57,47 @@ void FConsoleVariablesEditorModule::ShutdownModule()
 	}
 }
 
+void FConsoleVariablesEditorModule::SavePreset() const
+{
+	if (MainPanel.IsValid())
+	{
+		MainPanel->SaveCurrentPreset();
+	}
+}
+
+void FConsoleVariablesEditorModule::SaveSpecificPreset(const TObjectPtr<UConsoleVariablesAsset> Preset) const
+{
+	if (MainPanel.IsValid())
+	{
+		MainPanel->SaveSpecificPreset(Preset);
+	}
+}
+
+void FConsoleVariablesEditorModule::SavePresetAs() const
+{
+	if (MainPanel.IsValid())
+	{
+		MainPanel->SaveCurrentPresetAs();
+	}
+}
+
+void FConsoleVariablesEditorModule::SaveSpecificPresetAs(const TObjectPtr<UConsoleVariablesAsset> Preset) const
+{
+	if (MainPanel.IsValid())
+	{
+		MainPanel->SaveSpecificPresetAs(Preset);
+	}
+}
+
+void FConsoleVariablesEditorModule::OpenConsoleVariablesDialogWithPreset(
+	const TObjectPtr<UConsoleVariablesAsset> Preset) const
+{
+	if (MainPanel.IsValid())
+	{
+		MainPanel->ImportPreset(Preset);
+	}
+}
+
 void FConsoleVariablesEditorModule::OpenConsoleVariablesDialogWithAssetSelected(const FAssetData& InAssetData) const
 {
 	if (InAssetData.IsValid())
@@ -121,13 +162,30 @@ TArray<TWeakPtr<FConsoleVariablesEditorCommandInfo>> FConsoleVariablesEditorModu
 	
 	for (const TSharedPtr<FConsoleVariablesEditorCommandInfo>& CommandInfo : ConsoleObjectsMasterReference)
 	{
-		if (Algo::AnyOf(InTokens,
-			[&CommandInfo, InSearchCase](const FString& Token)
-			{
-				return CommandInfo->Command.Contains(Token, InSearchCase);
-			}))
+		// Match any
+		for (const FString& Token : InTokens)
 		{
-			ReturnValue.Add(CommandInfo);
+			bool bMatchFound = false;
+			
+			// Match all of these
+			const FString SpaceDelimiter = " ";
+			TArray<FString> OutSpacedArray;
+			if (Token.Contains(SpaceDelimiter) && Token.ParseIntoArray(OutSpacedArray, *SpaceDelimiter, true) > 1)
+			{
+				bMatchFound = Algo::AllOf(OutSpacedArray, [&CommandInfo, InSearchCase](const FString& Comparator)
+				{
+					return CommandInfo->Command.Contains(Comparator, InSearchCase);
+				});
+			}
+			else
+			{
+				bMatchFound = CommandInfo->Command.Contains(Token, InSearchCase);
+			}
+
+			if (bMatchFound)
+			{
+				ReturnValue.Add(CommandInfo);
+			}
 		}
 	}
 
@@ -157,6 +215,26 @@ TObjectPtr<UConsoleVariablesAsset> FConsoleVariablesEditorModule::GetGlobalSearc
 	return EditingGlobalSearchAsset;
 }
 
+FReply FConsoleVariablesEditorModule::ValidateConsoleInputAndAddToCurrentPreset(const FText& CommittedText) const
+{
+	return MainPanel->ValidateConsoleInputAndAddToCurrentPreset(CommittedText);
+}
+
+void FConsoleVariablesEditorModule::RefreshList() const
+{
+	MainPanel->RefreshList();
+}
+
+void FConsoleVariablesEditorModule::RebuildList() const
+{
+	MainPanel->RebuildList();
+}
+
+void FConsoleVariablesEditorModule::UpdatePresetValuesForSave(TObjectPtr<UConsoleVariablesAsset> InAsset)
+{
+	MainPanel->UpdatePresetValuesForSave(InAsset);
+}
+
 bool FConsoleVariablesEditorModule::PopulateGlobalSearchAssetWithVariablesMatchingTokens(const TArray<FString>& InTokens)
 {
 	// Remove existing commands
@@ -281,7 +359,7 @@ void FConsoleVariablesEditorModule::OnConsoleVariableChanged(IConsoleVariable* C
 			{
 				if (MainPanel.IsValid())
 				{
-					MainPanel->AddConsoleObjectToPreset(
+					MainPanel->AddConsoleObjectToCurrentPreset(
 						Key,
 						// If we're not in preset mode then pass empty value
 						// This forces the row to get the current value at the time it's generated
diff --git a/Engine/Plugins/Experimental/ConsoleVariables/Source/ConsoleVariablesEditor/Private/ConsoleVariablesEditorModule.h b/Engine/Plugins/Experimental/ConsoleVariables/Source/ConsoleVariablesEditor/Private/ConsoleVariablesEditorModule.h
--- a/Engine/Plugins/Experimental/ConsoleVariables/Source/ConsoleVariablesEditor/Private/ConsoleVariablesEditorModule.h
+++ b/Engine/Plugins/Experimental/ConsoleVariables/Source/ConsoleVariablesEditor/Private/ConsoleVariablesEditorModule.h
@@ -28,6 +28,11 @@ public:
 	virtual void ShutdownModule() override;
 	//~ End IModuleInterface Interface
 
+	void SavePreset() const;
+	void SaveSpecificPreset(const TObjectPtr<UConsoleVariablesAsset> Preset) const;
+	void SavePresetAs() const;
+	void SaveSpecificPresetAs(const TObjectPtr<UConsoleVariablesAsset> Preset) const;
+	void OpenConsoleVariablesDialogWithPreset(const TObjectPtr<UConsoleVariablesAsset> Preset) const;
 	void OpenConsoleVariablesDialogWithAssetSelected(const FAssetData& InAssetData) const;
 
 	/** Find all console variables and cache their startup values */
@@ -42,7 +47,11 @@ public:
 	TWeakPtr<FConsoleVariablesEditorCommandInfo> FindCommandInfoByName(
 		const FString& NameToSearch, ESearchCase::Type InSearchCase = ESearchCase::IgnoreCase);
 
-	/** Find all tracked console variables matching a specific search query with optional case sensitivity. */
+	/*Find all tracked console variables matching a specific search query with optional case sensitivity.
+	 *Individual members of InTokens will be considered "AnyOf" or "OR" searches. If SearchTerms contains any individual member it will match.
+	 *Members will be tested for a space character (" "). If a space is found, a subsearch will be run.
+	 *This subsearch will be an "AllOf" or "AND" type search in which all strings, separated by a space, must be found in the search terms.
+	 */
 	TArray<TWeakPtr<FConsoleVariablesEditorCommandInfo>> FindCommandInfosMatchingTokens(
 		const TArray<FString>& InTokens, ESearchCase::Type InSearchCase = ESearchCase::IgnoreCase);
 
@@ -56,6 +65,12 @@ public:
 	
 	[[nodiscard]] TObjectPtr<UConsoleVariablesAsset> GetPresetAsset() const;
 	[[nodiscard]] TObjectPtr<UConsoleVariablesAsset> GetGlobalSearchAsset() const;
+	
+	FReply ValidateConsoleInputAndAddToCurrentPreset(const FText& CommittedText) const;
+	void RefreshList() const;
+	void RebuildList() const;
+	
+	void UpdatePresetValuesForSave(TObjectPtr<UConsoleVariablesAsset> InAsset);
 
 	/** Fills Global Search Asset's Saved Commands with variables matching the specified query. Returns false if no matches were found. */
 	bool PopulateGlobalSearchAssetWithVariablesMatchingTokens(const TArray<FString>& InTokens);
diff --git a/Engine/Plugins/Experimental/ConsoleVariables/Source/ConsoleVariablesEditor/Private/ConsoleVariablesEditorStyle.cpp b/Engine/Plugins/Experimental/ConsoleVariables/Source/ConsoleVariablesEditor/Private/ConsoleVariablesEditorStyle.cpp
--- a/Engine/Plugins/Experimental/ConsoleVariables/Source/ConsoleVariablesEditor/Private/ConsoleVariablesEditorStyle.cpp
+++ b/Engine/Plugins/Experimental/ConsoleVariables/Source/ConsoleVariablesEditor/Private/ConsoleVariablesEditorStyle.cpp
@@ -52,6 +52,7 @@ const FSlateBrush* FConsoleVariablesEditorStyle::GetBrush(const FName PropertyNa
 #define IMAGE_BRUSH( RelativePath, ... ) FSlateImageBrush( Style->RootToContentDir( RelativePath, TEXT(".png") ), __VA_ARGS__ )
 #define BOX_BRUSH( RelativePath, ... ) FSlateBoxBrush( Style->RootToContentDir( RelativePath, TEXT(".png") ), __VA_ARGS__ )
 #define BORDER_BRUSH( RelativePath, ... ) FSlateBorderBrush( Style->RootToContentDir( RelativePath, TEXT(".png") ), __VA_ARGS__ )
+#define IMAGE_BRUSH_SVG( RelativePath, ... ) FSlateVectorImageBrush( Style->RootToContentDir(RelativePath, TEXT(".svg") ), __VA_ARGS__)
 #define IMAGE_PLUGIN_BRUSH_SVG( RelativePath, ... ) FSlateVectorImageBrush( FConsoleVariablesEditorStyle::InContent(RelativePath, ".svg"), __VA_ARGS__)
 
 const FVector2D Icon64x64(64.f, 64.f);
@@ -81,8 +82,14 @@ TSharedRef< FSlateStyleSet > FConsoleVariablesEditorStyle::Create()
 	Style->SetCoreContentRoot(FPaths::EngineContentDir() / TEXT("Slate"));
 
 	// Icons
-	Style->Set("ConsoleVariables.ToolbarButton", new IMAGE_BRUSH("Icons/Icon40", Icon40x40));
-	Style->Set("ConsoleVariables.ToolbarButton.Small", new IMAGE_BRUSH("Icons/Icon20", Icon20x20));
+	Style->Set("ConsoleVariables.ToolbarButton", new IMAGE_BRUSH_SVG("Icons/ConsoleVariables", Icon40x40));
+	Style->Set("ConsoleVariables.ToolbarButton.Small", new IMAGE_BRUSH_SVG("Icons/ConsoleVariables", Icon20x20));
+	
+	Style->Set("ConsoleVariables.Favorite.Outline", new IMAGE_BRUSH_SVG("Icons/FavoriteOutline", Icon40x40));
+	Style->Set("ConsoleVariables.Favorite.Outline.Small", new IMAGE_BRUSH_SVG("Icons/FavoriteOutline", Icon16x16));
+	
+	Style->Set("ConsoleVariables.GlobalSearch", new IMAGE_BRUSH_SVG("Icons/SearchGlobal", Icon40x40));
+	Style->Set("ConsoleVariables.GlobalSearch.Small", new IMAGE_BRUSH_SVG("Icons/SearchGlobal", Icon16x16));
 
 	// Brush
 	Style->Set("ConsoleVariablesEditor.GroupBorder", new BOX_BRUSH("Common/DarkGroupBorder", FMargin(4.0f / 16.0f)));
diff --git a/Engine/Plugins/Experimental/ConsoleVariables/Source/ConsoleVariablesEditor/Private/Views/List/ConsoleVariablesEditorList.cpp b/Engine/Plugins/Experimental/ConsoleVariables/Source/ConsoleVariablesEditor/Private/Views/List/ConsoleVariablesEditorList.cpp
--- a/Engine/Plugins/Experimental/ConsoleVariables/Source/ConsoleVariablesEditor/Private/Views/List/ConsoleVariablesEditorList.cpp
+++ b/Engine/Plugins/Experimental/ConsoleVariables/Source/ConsoleVariablesEditor/Private/Views/List/ConsoleVariablesEditorList.cpp
@@ -19,6 +19,11 @@ TSharedRef<SWidget> FConsoleVariablesEditorList::GetOrCreateWidget()
 	return ListWidget.ToSharedRef();
 }
 
+void FConsoleVariablesEditorList::SetSearchString(const FString& SearchString)
+{
+	ListWidget->SetSearchStringInSearchInputField(SearchString);
+}
+
 void FConsoleVariablesEditorList::RebuildList(const FString& InConsoleCommandToScrollTo) const
 {
 	if (ListWidget.IsValid())
diff --git a/Engine/Plugins/Experimental/ConsoleVariables/Source/ConsoleVariablesEditor/Private/Views/List/ConsoleVariablesEditorList.h b/Engine/Plugins/Experimental/ConsoleVariables/Source/ConsoleVariablesEditor/Private/Views/List/ConsoleVariablesEditorList.h
--- a/Engine/Plugins/Experimental/ConsoleVariables/Source/ConsoleVariablesEditor/Private/Views/List/ConsoleVariablesEditorList.h
+++ b/Engine/Plugins/Experimental/ConsoleVariables/Source/ConsoleVariablesEditor/Private/Views/List/ConsoleVariablesEditorList.h
@@ -28,6 +28,8 @@ public:
 
 	TSharedRef<SWidget> GetOrCreateWidget();
 
+	void SetSearchString(const FString& SearchString);
+
 	EConsoleVariablesEditorListMode GetListMode() const
 	{
 		return CurrentListMode;
diff --git a/Engine/Plugins/Experimental/ConsoleVariables/Source/ConsoleVariablesEditor/Private/Views/List/ConsoleVariablesEditorListRow.cpp b/Engine/Plugins/Experimental/ConsoleVariables/Source/ConsoleVariablesEditor/Private/Views/List/ConsoleVariablesEditorListRow.cpp
--- a/Engine/Plugins/Experimental/ConsoleVariables/Source/ConsoleVariablesEditor/Private/Views/List/ConsoleVariablesEditorListRow.cpp
+++ b/Engine/Plugins/Experimental/ConsoleVariables/Source/ConsoleVariablesEditor/Private/Views/List/ConsoleVariablesEditorListRow.cpp
@@ -5,6 +5,7 @@
 #include "ConsoleVariablesAsset.h"
 #include "ConsoleVariablesEditorModule.h"
 
+#include "Algo/AllOf.h"
 #include "Algo/AnyOf.h"
 #include "Views/List/ConsoleVariablesEditorList.h"
 #include "Views/List/SConsoleVariablesEditorList.h"
@@ -135,8 +136,8 @@ void FConsoleVariablesEditorListRow::SetPresetValue(const FString& InPresetValue
 	PresetValue = InPresetValue;
 }
 
-bool FConsoleVariablesEditorListRow::MatchSearchTokensToSearchTerms(const TArray<FString> InTokens,
-                                                                    const bool bMatchAnyTokens)
+bool FConsoleVariablesEditorListRow::MatchSearchTokensToSearchTerms(
+	const TArray<FString> InTokens, ESearchCase::Type InSearchCase)
 {
 	// If the search is cleared we'll consider the row to pass search
 	bool bMatchFound = InTokens.Num() == 0;
@@ -153,11 +154,29 @@ bool FConsoleVariablesEditorListRow::MatchSearchTokensToSearchTerms(const TArray
 			SearchTerms += AsVariable->GetString() + PinnedInfo->GetHelpText();
 		}
 
-		bMatchFound = Algo::AnyOf(InTokens,
-		                          [&SearchTerms](const FString& Token)
-		                          {
-			                          return SearchTerms.Contains(Token);
-		                          });
+		// Match any
+		for (const FString& Token : InTokens)
+		{
+			// Match all of these
+			const FString SpaceDelimiter = " ";
+			TArray<FString> OutSpacedArray;
+			if (Token.Contains(SpaceDelimiter) && Token.ParseIntoArray(OutSpacedArray, *SpaceDelimiter, true) > 1)
+			{
+				bMatchFound = Algo::AllOf(OutSpacedArray, [&SearchTerms, InSearchCase](const FString& Comparator)
+				{
+					return SearchTerms.Contains(Comparator, InSearchCase);
+				});
+			}
+			else
+			{
+				bMatchFound = SearchTerms.Contains(Token, InSearchCase);
+			}
+
+			if (bMatchFound)
+			{
+				break;
+			}
+		}
 	}
 
 	bDoesRowMatchSearchTerms = bMatchFound;
@@ -302,13 +321,21 @@ FReply FConsoleVariablesEditorListRow::OnActionButtonClicked()
 
 	if (bIsGlobalSearch)
 	{
-		EditableAsset->AddOrSetConsoleObjectSavedData(
-			{
-				CommandName,
-				"",
-				ECheckBoxState::Checked
-			}
-		);
+		FConsoleVariablesEditorAssetSaveData MatchingData;
+		if (!ConsoleVariablesEditorModule.GetPresetAsset()->FindSavedDataByCommandString(CommandName, MatchingData))
+		{
+			EditableAsset->AddOrSetConsoleObjectSavedData(
+				{
+					CommandName,
+					"",
+					ECheckBoxState::Checked
+				}
+			);
+		}
+		else
+		{
+			EditableAsset->RemoveConsoleVariable(CommandName);
+		}
 	}
 	else
 	{
diff --git a/Engine/Plugins/Experimental/ConsoleVariables/Source/ConsoleVariablesEditor/Private/Views/List/ConsoleVariablesEditorListRow.h b/Engine/Plugins/Experimental/ConsoleVariables/Source/ConsoleVariablesEditor/Private/Views/List/ConsoleVariablesEditorListRow.h
--- a/Engine/Plugins/Experimental/ConsoleVariables/Source/ConsoleVariablesEditor/Private/Views/List/ConsoleVariablesEditorListRow.h
+++ b/Engine/Plugins/Experimental/ConsoleVariables/Source/ConsoleVariablesEditor/Private/Views/List/ConsoleVariablesEditorListRow.h
@@ -84,8 +84,12 @@ struct FConsoleVariablesEditorListRow final : TSharedFromThis<FConsoleVariablesE
 	[[nodiscard]] const FString& GetPresetValue() const;
 	void SetPresetValue(const FString& InPresetValue);
 
-	/* If bMatchAnyTokens is false, only nodes that match all terms will be returned. */
-	bool MatchSearchTokensToSearchTerms(const TArray<FString> InTokens, const bool bMatchAnyTokens = false);
+	/*
+	 *Individual members of InTokens will be considered "AnyOf" or "OR" searches. If SearchTerms contains any individual member it will match.
+	 *Members will be tested for a space character (" "). If a space is found, a subsearch will be run.
+	 *This subsearch will be an "AllOf" or "AND" type search in which all strings, separated by a space, must be found in the search terms.
+	 */
+	bool MatchSearchTokensToSearchTerms(const TArray<FString> InTokens, ESearchCase::Type InSearchCase = ESearchCase::IgnoreCase);
 
 	/* This overload creates tokens from a string first, then calls ExecuteSearchOnChildNodes(const TArray<FString>& Tokens). */
 	void ExecuteSearchOnChildNodes(const FString& SearchString) const;
diff --git a/Engine/Plugins/Experimental/ConsoleVariables/Source/ConsoleVariablesEditor/Private/Views/List/SConsoleVariablesEditorList.cpp b/Engine/Plugins/Experimental/ConsoleVariables/Source/ConsoleVariablesEditor/Private/Views/List/SConsoleVariablesEditorList.cpp
--- a/Engine/Plugins/Experimental/ConsoleVariables/Source/ConsoleVariablesEditor/Private/Views/List/SConsoleVariablesEditorList.cpp
+++ b/Engine/Plugins/Experimental/ConsoleVariables/Source/ConsoleVariablesEditor/Private/Views/List/SConsoleVariablesEditorList.cpp
@@ -8,6 +8,7 @@
 #include "ConsoleVariablesEditorListFilters/ConsoleVariablesEditorListFilter_SourceText.h"
 #include "ConsoleVariablesEditorModule.h"
 #include "ConsoleVariablesEditorProjectSettings.h"
+#include "ConsoleVariablesEditorStyle.h"
 #include "SConsoleVariablesEditorListRow.h"
 #include "../Widgets/SConsoleVariablesEditorGlobalSearchToggle.h"
 
@@ -91,13 +92,12 @@ void SConsoleVariablesEditorList::Construct(const FArguments& InArgs, TSharedRef
 					.Padding(0, 1, 4, 0)
 					.AutoWidth()
 					[
-						SNew(STextBlock)
-						.TextStyle(FAppStyle::Get(), "FindResults.FindInBlueprints")
-						.Text(FText::FromString(FString(TEXT("\xf1e5"))) /*fa-binoculars*/)
+						SNew(SImage)
+						.Image(FConsoleVariablesEditorStyle::Get().GetBrush("ConsoleVariables.GlobalSearch.Small"))
 					]
 
 					+SHorizontalBox::Slot()
-					.Padding(4.f, 0.f, 0.f, 0.f)
+					.Padding(4.f, 1.f, 0.f, 0.f)
 					.AutoWidth()
 					[
 						SNew(STextBlock)
@@ -145,6 +145,7 @@ void SConsoleVariablesEditorList::Construct(const FArguments& InArgs, TSharedRef
 			.HAlign(HAlign_Fill)
 			[
 				SAssignNew(GlobalSearchesContainer, SWrapBox)
+				.InnerSlotPadding(FVector2d(6, 4))
 				.UseAllottedSize(true)
 			]
 
@@ -226,7 +227,8 @@ void SConsoleVariablesEditorList::Construct(const FArguments& InArgs, TSharedRef
 				SNew(STextBlock)
 				.AutoWrapText(true)
 				.Justification(ETextJustify::Center)
-				.Text(LOCTEXT("ConsoleVariablesEditorList_NoList", "No List to show. Try clearing the active search or adding some console variables to the list."))
+				.Text(LOCTEXT("ConsoleVariablesEditorList_NoList",
+					"No List to show. Try clearing the active local/global search and/or adding some console variables to the list."))
 			]
 		]
 	];
@@ -275,11 +277,12 @@ FReply SConsoleVariablesEditorList::TryEnterGlobalSearch(const FString& SearchSt
 		return ReturnValue;
 	}
 
-	// Strings that already have associated buttons
+	// Strings that already have associated buttons	
 	TArray<FString> ExistingButtonStrings;
+	
 	// All strings parsed from the search text
 	TArray<FString> OutTokens;
-	SearchString.ParseIntoArray(OutTokens, TEXT(" "));
+	SearchString.ParseIntoArray(OutTokens, TEXT("|"), true);
 
 	// Get tokens from current searches. This step allows us to properly populate the asset with all matching commands
 	for (const TSharedRef<SConsoleVariablesEditorGlobalSearchToggle>& GlobalSearchButton : CurrentGlobalSearches)
@@ -305,7 +308,7 @@ FReply SConsoleVariablesEditorList::TryEnterGlobalSearch(const FString& SearchSt
 	
 	FConsoleVariablesEditorModule& ConsoleVariablesEditorModule = FConsoleVariablesEditorModule::Get();
 	
-	if (ConsoleVariablesEditorModule.PopulateGlobalSearchAssetWithVariablesMatchingTokens(OutTokens))
+	const bool bFoundMatches = ConsoleVariablesEditorModule.PopulateGlobalSearchAssetWithVariablesMatchingTokens(OutTokens);
 	{
 		// If we were in Preset mode before entering global search, cache the existing tree objects to maintain state
 		if (ListModelPtr.Pin()->GetListMode() == FConsoleVariablesEditorList::EConsoleVariablesEditorListMode::Preset)
@@ -338,7 +341,8 @@ FReply SConsoleVariablesEditorList::TryEnterGlobalSearch(const FString& SearchSt
 
 		ReturnValue = FReply::Handled();
 	}
-	else
+	
+	if (!bFoundMatches)
 	{
 		UE_LOG(LogConsoleVariablesEditor, Warning,
 			TEXT("%hs: Failed to find console variable objects with names containing search string %s"),
@@ -478,13 +482,22 @@ FString SConsoleVariablesEditorList::GetSearchStringFromSearchInputField() const
 	? ListSearchBoxPtr->GetText().ToString() : "";
 }
 
+void SConsoleVariablesEditorList::SetSearchStringInSearchInputField(const FString InSearchString) const
+{
+	if (ensureAlwaysMsgf(ListSearchBoxPtr.IsValid(),
+		TEXT("%hs: ListSearchBoxPtr is not valid. Check to make sure it was created."), __FUNCTION__))
+	{
+		ListSearchBoxPtr->SetText(FText::FromString(InSearchString));
+	}
+}
+
 void SConsoleVariablesEditorList::ExecuteListViewSearchOnAllRows(
 	const FString& SearchString, const bool bShouldRefreshAfterward)
 {
 	TArray<FString> Tokens;
 	
 	// unquoted search equivalent to a match-any-of search
-	SearchString.ParseIntoArray(Tokens, TEXT(" "), true);
+	SearchString.ParseIntoArray(Tokens, TEXT("|"), true);
 	
 	for (const TSharedPtr<FConsoleVariablesEditorListRow>& ChildRow : TreeViewRootObjects)
 	{
@@ -958,7 +971,7 @@ void SConsoleVariablesEditorList::SetAllGroupsCollapsed()
 
 void SConsoleVariablesEditorList::OnListViewSearchTextChanged(const FText& Text)
 {
-	ExecuteListViewSearchOnAllRows(Text.ToString());
+	ExecuteListViewSearchOnAllRows(Text.ToString(), false);
 }
 
 void SConsoleVariablesEditorList::GenerateTreeView()
diff --git a/Engine/Plugins/Experimental/ConsoleVariables/Source/ConsoleVariablesEditor/Private/Views/List/SConsoleVariablesEditorList.h b/Engine/Plugins/Experimental/ConsoleVariables/Source/ConsoleVariablesEditor/Private/Views/List/SConsoleVariablesEditorList.h
--- a/Engine/Plugins/Experimental/ConsoleVariables/Source/ConsoleVariablesEditor/Private/Views/List/SConsoleVariablesEditorList.h
+++ b/Engine/Plugins/Experimental/ConsoleVariables/Source/ConsoleVariablesEditor/Private/Views/List/SConsoleVariablesEditorList.h
@@ -67,6 +67,7 @@ public:
 	void UpdatePresetValuesForSave(const TObjectPtr<UConsoleVariablesAsset> InAsset) const;
 
 	FString GetSearchStringFromSearchInputField() const;
+	void SetSearchStringInSearchInputField(const FString InSearchString) const;
 	void ExecuteListViewSearchOnAllRows(const FString& SearchString, const bool bShouldRefreshAfterward = true);
 
 	bool DoesTreeViewHaveVisibleChildren() const;
diff --git a/Engine/Plugins/Experimental/ConsoleVariables/Source/ConsoleVariablesEditor/Private/Views/List/SConsoleVariablesEditorListRow.cpp b/Engine/Plugins/Experimental/ConsoleVariables/Source/ConsoleVariablesEditor/Private/Views/List/SConsoleVariablesEditorListRow.cpp
--- a/Engine/Plugins/Experimental/ConsoleVariables/Source/ConsoleVariablesEditor/Private/Views/List/SConsoleVariablesEditorListRow.cpp
+++ b/Engine/Plugins/Experimental/ConsoleVariables/Source/ConsoleVariablesEditor/Private/Views/List/SConsoleVariablesEditorListRow.cpp
@@ -101,10 +101,10 @@ TSharedRef<SWidget> SConsoleVariablesEditorListRow::GenerateWidgetForColumn(cons
 
 	const TSharedRef<SImage> FlashImage = SNew(SImage)
 										.Image(new FSlateColorBrush(FStyleColors::White))
-										.Visibility_Raw(this, &SConsoleVariablesEditorListRow::GetFlashImageVisibility)
+										.Visibility_Raw(
+											this, &SConsoleVariablesEditorListRow::GetFlashImageVisibility)
 										.ColorAndOpacity_Raw(
-		                                                  this,
-		                                                  &SConsoleVariablesEditorListRow::GetFlashImageColorAndOpacity);
+											this, &SConsoleVariablesEditorListRow::GetFlashImageColorAndOpacity);
 
 	FlashImages.Add(FlashImage);
 
@@ -134,29 +134,8 @@ void SConsoleVariablesEditorListRow::OnMouseEnter(const FGeometry& MyGeometry, c
 
 	if (HoverableWidgetsPtr.IsValid())
 	{
-		// We don't want to show the widget if we're in Global Search mode and the command is already in the preset
-		const bool bIsListValid =
-			Item.Pin()->GetListViewPtr().IsValid() && Item.Pin()->GetListViewPtr().Pin()->GetListModelPtr().IsValid();
-
-		const FConsoleVariablesEditorList::EConsoleVariablesEditorListMode ListMode =
-			bIsListValid
-				? Item.Pin()->GetListViewPtr().Pin()->GetListModelPtr().Pin()->GetListMode()
-				: FConsoleVariablesEditorList::EConsoleVariablesEditorListMode::Preset;
-
-		const bool bIsGlobalSearch =
-			ListMode == FConsoleVariablesEditorList::EConsoleVariablesEditorListMode::GlobalSearch;
-
-		const FConsoleVariablesEditorModule& ConsoleVariablesEditorModule = FConsoleVariablesEditorModule::Get();
-
-		const FString& CommandName = Item.Pin()->GetCommandInfo().Pin()->Command;
-		FConsoleVariablesEditorAssetSaveData MatchingData;
-		if (!bIsGlobalSearch ||
-			(bIsGlobalSearch &&
-				!ConsoleVariablesEditorModule.GetPresetAsset()->
-				                              FindSavedDataByCommandString(CommandName, MatchingData)))
-		{
-			HoverableWidgetsPtr->SetVisibility(EVisibility::SelfHitTestInvisible);
-		}
+		HoverableWidgetsPtr->DetermineButtonImageAndTooltip();
+		HoverableWidgetsPtr->SetVisibility(EVisibility::SelfHitTestInvisible);
 	}
 
 	SMultiColumnTableRow<FConsoleVariablesEditorListRowPtr>::OnMouseEnter(MyGeometry, MouseEvent);
@@ -587,28 +566,15 @@ void SConsoleVariablesEditorListRowHoverWidgets::Construct(const FArguments& InA
 
 	Item = InRow;
 
-	const bool bIsListValid =
-		Item.Pin()->GetListViewPtr().IsValid() && Item.Pin()->GetListViewPtr().Pin()->GetListModelPtr().IsValid();
-
-	const FConsoleVariablesEditorList::EConsoleVariablesEditorListMode ListMode =
-		bIsListValid
-			? Item.Pin()->GetListViewPtr().Pin()->GetListModelPtr().Pin()->GetListMode()
-			: FConsoleVariablesEditorList::EConsoleVariablesEditorListMode::Preset;
-
-	FText ButtonTooltip = LOCTEXT("RemoveCvarTooltip",
-	                              "Remove cvar from this list and reset its value to the startup value.");
-	FName ButtonImageName = "Icons.Delete";
-
-	if (ListMode == FConsoleVariablesEditorList::EConsoleVariablesEditorListMode::GlobalSearch)
-	{
-		ButtonImageName = "Icons.Star";
-		ButtonTooltip = LOCTEXT("AddCvarToPresetTooltip", "Add this cvar to your current preset.");
-	}
+	const FText ButtonTooltip = LOCTEXT("RemoveCvarTooltip",
+								  "Remove cvar from this list and reset its value to the startup value.");
+	const FSlateBrush* ButtonImage = FAppStyle::Get().GetBrush("Icons.Delete");
 
 	ChildSlot
 	[
 		// Action Button
 		SAssignNew(ActionButtonPtr, SButton)
+		.ButtonStyle(&FAppStyle::Get().GetWidgetStyle<FButtonStyle>("NoBorder"))
 		.Visibility(EVisibility::Visible)
 		.ToolTipText(ButtonTooltip)
 		.ButtonColorAndOpacity(FStyleColors::Transparent)
@@ -621,15 +587,50 @@ void SConsoleVariablesEditorListRowHoverWidgets::Construct(const FArguments& InA
 		[
 			SNew(SScaleBox)
 			[
-				SNew(SImage)
+				SAssignNew(ActionButtonImage, SImage)
 				.Visibility(EVisibility::SelfHitTestInvisible)
-				.Image(FAppStyle::Get().GetBrush(ButtonImageName))
+				.Image(ButtonImage)
 				.ColorAndOpacity(FSlateColor::UseForeground())
 			]
 		]
 	];
 }
 
+void SConsoleVariablesEditorListRowHoverWidgets::DetermineButtonImageAndTooltip()
+{
+	const bool bIsListValid =
+		Item.Pin()->GetListViewPtr().IsValid() && Item.Pin()->GetListViewPtr().Pin()->GetListModelPtr().IsValid();
+
+	const FConsoleVariablesEditorList::EConsoleVariablesEditorListMode ListMode =
+		bIsListValid
+			? Item.Pin()->GetListViewPtr().Pin()->GetListModelPtr().Pin()->GetListMode()
+			: FConsoleVariablesEditorList::EConsoleVariablesEditorListMode::Preset;
+
+	const bool bIsGlobalSearch =
+		ListMode == FConsoleVariablesEditorList::EConsoleVariablesEditorListMode::GlobalSearch;
+
+	if (bIsGlobalSearch)
+	{
+		FText ButtonTooltip = LOCTEXT("AddCvarToPresetTooltip", "Add this cvar to your current preset.");
+		const FSlateBrush* ButtonImage = FConsoleVariablesEditorStyle::Get().GetBrush("ConsoleVariables.Favorite.Outline.Small");
+		
+		const FConsoleVariablesEditorModule& ConsoleVariablesEditorModule = FConsoleVariablesEditorModule::Get();
+		
+		const FString& CommandName = Item.Pin()->GetCommandInfo().Pin()->Command;
+		FConsoleVariablesEditorAssetSaveData MatchingData;
+
+		// If the item exists in the current preset already
+		if (ConsoleVariablesEditorModule.GetPresetAsset()->FindSavedDataByCommandString(CommandName, MatchingData))
+		{
+			ButtonTooltip = LOCTEXT("RemoveCvarFromPresetTooltip", "Remove this cvar from your current preset.");
+			ButtonImage = FAppStyle::Get().GetBrush("Icons.Star");
+		}
+
+		ActionButtonPtr->SetToolTipText(ButtonTooltip);
+		ActionButtonImage->SetImage(ButtonImage);
+	}
+}
+
 void SConsoleVariablesEditorListRowHoverWidgets::OnMouseEnter(const FGeometry& MyGeometry,
                                                               const FPointerEvent& MouseEvent)
 {
diff --git a/Engine/Plugins/Experimental/ConsoleVariables/Source/ConsoleVariablesEditor/Private/Views/List/SConsoleVariablesEditorListRow.h b/Engine/Plugins/Experimental/ConsoleVariables/Source/ConsoleVariablesEditor/Private/Views/List/SConsoleVariablesEditorListRow.h
--- a/Engine/Plugins/Experimental/ConsoleVariables/Source/ConsoleVariablesEditor/Private/Views/List/SConsoleVariablesEditorListRow.h
+++ b/Engine/Plugins/Experimental/ConsoleVariables/Source/ConsoleVariablesEditor/Private/Views/List/SConsoleVariablesEditorListRow.h
@@ -101,6 +101,8 @@ public:
 	
 	void Construct(const FArguments& InArgs, const TWeakPtr<FConsoleVariablesEditorListRow> InRow);
 
+	void DetermineButtonImageAndTooltip();
+
 	virtual void OnMouseEnter(const FGeometry& MyGeometry, const FPointerEvent& MouseEvent) override;
 
 	virtual void OnMouseLeave(const FPointerEvent& MouseEvent) override;
@@ -112,4 +114,5 @@ private:
 	TWeakPtr<FConsoleVariablesEditorListRow> Item;
 	
 	TSharedPtr<SButton> ActionButtonPtr;
+	TSharedPtr<SImage> ActionButtonImage;
 };
diff --git a/Engine/Plugins/Experimental/ConsoleVariables/Source/ConsoleVariablesEditor/Private/Views/MainPanel/ConsoleVariablesEditorMainPanel.cpp b/Engine/Plugins/Experimental/ConsoleVariables/Source/ConsoleVariablesEditor/Private/Views/MainPanel/ConsoleVariablesEditorMainPanel.cpp
--- a/Engine/Plugins/Experimental/ConsoleVariables/Source/ConsoleVariablesEditor/Private/Views/MainPanel/ConsoleVariablesEditorMainPanel.cpp
+++ b/Engine/Plugins/Experimental/ConsoleVariables/Source/ConsoleVariablesEditor/Private/Views/MainPanel/ConsoleVariablesEditorMainPanel.cpp
@@ -53,7 +53,7 @@ TObjectPtr<UConsoleVariablesAsset> FConsoleVariablesEditorMainPanel::GetEditingA
 	return GetConsoleVariablesModule().GetPresetAsset();
 }
 
-void FConsoleVariablesEditorMainPanel::AddConsoleObjectToPreset(
+void FConsoleVariablesEditorMainPanel::AddConsoleObjectToCurrentPreset(
 	const FString InConsoleCommand, const FString InValue, const bool bScrollToNewRow) const
 {
 	if (const TObjectPtr<UConsoleVariablesAsset> Asset = GetEditingAsset())
@@ -75,6 +75,70 @@ void FConsoleVariablesEditorMainPanel::AddConsoleObjectToPreset(
 	}
 }
 
+FReply FConsoleVariablesEditorMainPanel::ValidateConsoleInputAndAddToCurrentPreset(const FText& CommittedText) const
+{
+	// Clear Search
+	if (const TSharedPtr<FConsoleVariablesEditorList> PinnedList = GetEditorList().Pin())
+	{
+		PinnedList->SetSearchString("");
+	}
+	
+	const FString CommandString = CommittedText.ToString().TrimStartAndEnd();
+	FString CommandKey;
+	FString ValueString;
+	
+	if (CommandString.Contains(" "))
+	{
+		CommandString.Split(TEXT(" "), &CommandKey, &ValueString);
+	}
+	else
+	{
+		CommandKey = CommandString;
+	}
+
+	if (IConsoleObject* ConsoleObject = IConsoleManager::Get().FindConsoleObject(*CommandKey)) 
+	{
+		if (ValueString.IsEmpty())
+		{
+			if (const IConsoleVariable* AsVariable = ConsoleObject->AsVariable())
+			{
+				ValueString = AsVariable->GetString();
+			}
+		}
+		
+		AddConsoleObjectToCurrentPreset(
+			CommandKey,
+			ValueString,
+			true
+		);
+	}
+	else if (CommandString.IsEmpty())
+	{
+		UE_LOG(LogConsoleVariablesEditor, Warning, TEXT("%hs: Input is blank."), __FUNCTION__);
+
+		return FReply::Unhandled();
+	}
+	// Try to execute the whole command. Some commands are not registered, but are parsed externally so they won't be found by IConsoleManager::FindConsoleObject
+	else if (GEngine->Exec(FConsoleVariablesEditorCommandInfo::GetCurrentWorld(), *CommandString))
+	{
+		AddConsoleObjectToCurrentPreset(
+			CommandString,
+			"",
+			true
+			);
+	}
+	else
+	{
+		UE_LOG(LogConsoleVariablesEditor, Warning,
+			TEXT("%hs: Input '%s' is not a recognized console variable or command."),
+			__FUNCTION__, *CommandString);
+
+		return FReply::Unhandled();
+	}
+
+	return FReply::Handled();
+}
+
 void FConsoleVariablesEditorMainPanel::RebuildList(const FString InConsoleCommandToScrollTo) const
 {
 	if (EditorList.IsValid())
@@ -107,11 +171,10 @@ void FConsoleVariablesEditorMainPanel::RefreshMultiUserDetails() const
 	}
 }
 
-void FConsoleVariablesEditorMainPanel::SavePreset()
+void FConsoleVariablesEditorMainPanel::SaveCurrentPreset()
 {
-	const TWeakObjectPtr<UConsoleVariablesAsset> EditingAsset = GetEditingAsset();
-	
-	if (ReferenceAssetOnDisk.IsValid() && EditingAsset.IsValid())
+	if (const TWeakObjectPtr<UConsoleVariablesAsset> EditingAsset = GetEditingAsset();
+		EditingAsset.IsValid() && ReferenceAssetOnDisk.IsValid())
 	{
 		UpdatePresetValuesForSave(EditingAsset.Get());
 		
@@ -126,14 +189,28 @@ void FConsoleVariablesEditorMainPanel::SavePreset()
 	}
 
 	// Fallback
-	SavePresetAs();
+	SaveCurrentPresetAs();
 }
 
-void FConsoleVariablesEditorMainPanel::SavePresetAs()
+void FConsoleVariablesEditorMainPanel::SaveSpecificPreset(const TObjectPtr<UConsoleVariablesAsset> InPreset) const
 {
-	const TWeakObjectPtr<UConsoleVariablesAsset> EditingAsset = GetEditingAsset();
-	
-	if (EditingAsset.IsValid())
+	if (InPreset)
+	{		
+		if (UPackage* ReferencePackage = InPreset->GetPackage())
+		{
+			UEditorLoadingAndSavingUtils::SavePackages({ ReferencePackage }, false);
+
+			return;
+		}
+	}
+
+	// Fallback
+	SaveSpecificPresetAs(InPreset);
+}
+
+void FConsoleVariablesEditorMainPanel::SaveCurrentPresetAs()
+{
+	if (const TWeakObjectPtr<UConsoleVariablesAsset> EditingAsset = GetEditingAsset(); EditingAsset.IsValid())
 	{
 		UpdatePresetValuesForSave(EditingAsset.Get());
 			
@@ -152,31 +229,49 @@ void FConsoleVariablesEditorMainPanel::SavePresetAs()
 	}
 }
 
+void FConsoleVariablesEditorMainPanel::SaveSpecificPresetAs(const TObjectPtr<UConsoleVariablesAsset> InPreset) const
+{
+	if (InPreset)
+	{
+		TArray<UObject*> SavedAssets;
+		FEditorFileUtils::SaveAssetsAs({ InPreset }, SavedAssets);
+	}
+}
+
 void FConsoleVariablesEditorMainPanel::ImportPreset(const FAssetData& InPresetAsset)
 {
 	FSlateApplication::Get().DismissAllMenus();
-	const TObjectPtr<UConsoleVariablesAsset> EditingAsset = GetEditingAsset();
+	
+	if (UConsoleVariablesAsset* Preset = CastChecked<UConsoleVariablesAsset>(InPresetAsset.GetAsset()))
+	{
+		if (const TObjectPtr<UConsoleVariablesAsset> EditingAsset = GetEditingAsset(); ImportPreset_Impl(Preset, EditingAsset))
+		{
+			EditorList->RebuildList();
+		}
+	}
+}
+
+void FConsoleVariablesEditorMainPanel::ImportPreset(const TObjectPtr<UConsoleVariablesAsset> InPreset)
+{
+	FSlateApplication::Get().DismissAllMenus();
 
-	if (EditingAsset && ImportPreset_Impl(InPresetAsset, EditingAsset))
+	if (const TObjectPtr<UConsoleVariablesAsset> EditingAsset = GetEditingAsset(); ImportPreset_Impl(InPreset, EditingAsset))
 	{
 		EditorList->RebuildList();
 	}
 }
 
 bool FConsoleVariablesEditorMainPanel::ImportPreset_Impl(
-	const FAssetData& InPresetAsset, const TObjectPtr<UConsoleVariablesAsset> EditingAsset)
+	const TObjectPtr<UConsoleVariablesAsset> Preset, const TObjectPtr<UConsoleVariablesAsset> EditingAsset)
 {
-	if (UConsoleVariablesAsset* Preset = CastChecked<UConsoleVariablesAsset>(InPresetAsset.GetAsset()))
+	if (Preset && EditingAsset)
 	{
-		if (EditingAsset)
-		{
-			ReferenceAssetOnDisk = Preset;
+		ReferenceAssetOnDisk = Preset;
 
-			EditingAsset->Modify();
-			EditingAsset->CopyFrom(Preset);
+		EditingAsset->Modify();
+		EditingAsset->CopyFrom(Preset);
 
-			return true;
-		}
+		return true;
 	}
 
 	return false;
diff --git a/Engine/Plugins/Experimental/ConsoleVariables/Source/ConsoleVariablesEditor/Private/Views/MainPanel/ConsoleVariablesEditorMainPanel.h b/Engine/Plugins/Experimental/ConsoleVariables/Source/ConsoleVariablesEditor/Private/Views/MainPanel/ConsoleVariablesEditorMainPanel.h
--- a/Engine/Plugins/Experimental/ConsoleVariables/Source/ConsoleVariablesEditor/Private/Views/MainPanel/ConsoleVariablesEditorMainPanel.h
+++ b/Engine/Plugins/Experimental/ConsoleVariables/Source/ConsoleVariablesEditor/Private/Views/MainPanel/ConsoleVariablesEditorMainPanel.h
@@ -22,10 +22,12 @@ public:
 	static FConsoleVariablesEditorModule& GetConsoleVariablesModule();
 	static TObjectPtr<UConsoleVariablesAsset> GetEditingAsset();
 
-	void AddConsoleObjectToPreset(
+	void AddConsoleObjectToCurrentPreset(
 		const FString InConsoleCommand, const FString InValue,
 		const bool bScrollToNewRow = false) const;
 
+	FReply ValidateConsoleInputAndAddToCurrentPreset(const FText& CommittedText) const;
+
 	/** Regenerate the list items and refresh the list. Call when adding or removing variables. */
 	void RebuildList(const FString InConsoleCommandToScrollTo = "") const;
 
@@ -41,9 +43,12 @@ public:
 
 	// Save / Load
 
-	void SavePreset();
-	void SavePresetAs();
+	void SaveCurrentPreset();
+	void SaveSpecificPreset(const TObjectPtr<UConsoleVariablesAsset> InPreset) const;
+	void SaveCurrentPresetAs();
+	void SaveSpecificPresetAs(const TObjectPtr<UConsoleVariablesAsset> InPreset) const;
 	void ImportPreset(const FAssetData& InPresetAsset);
+	void ImportPreset(const TObjectPtr<UConsoleVariablesAsset> InPreset);
 
 	TWeakObjectPtr<UConsoleVariablesAsset> GetReferenceAssetOnDisk() const
 	{
@@ -67,7 +72,7 @@ public:
 
 private:
 
-	bool ImportPreset_Impl(const FAssetData& InPresetAsset, const TObjectPtr<UConsoleVariablesAsset> EditingAsset);
+	bool ImportPreset_Impl(const TObjectPtr<UConsoleVariablesAsset> Preset, const TObjectPtr<UConsoleVariablesAsset> EditingAsset);
 
 	TSharedPtr<SConsoleVariablesEditorMainPanel> MainPanelWidget;
 
diff --git a/Engine/Plugins/Experimental/ConsoleVariables/Source/ConsoleVariablesEditor/Private/Views/MainPanel/SConsoleVariablesEditorCustomConsoleInputBox.cpp b/Engine/Plugins/Experimental/ConsoleVariables/Source/ConsoleVariablesEditor/Private/Views/MainPanel/SConsoleVariablesEditorCustomConsoleInputBox.cpp
--- a/Engine/Plugins/Experimental/ConsoleVariables/Source/ConsoleVariablesEditor/Private/Views/MainPanel/SConsoleVariablesEditorCustomConsoleInputBox.cpp
+++ b/Engine/Plugins/Experimental/ConsoleVariables/Source/ConsoleVariablesEditor/Private/Views/MainPanel/SConsoleVariablesEditorCustomConsoleInputBox.cpp
@@ -305,7 +305,7 @@ FReply SConsoleVariablesEditorCustomConsoleInputBox::OnKeyDownHandler(const FGeo
 		InputText->SetText(FText::GetEmpty());
 		SuggestionBox->SetIsOpen( false );
 		bIgnoreUIUpdate = false;
-		return MainPanelWidget.Pin()->ValidateConsoleInput(CommittedText);
+		return MainPanelWidget.Pin()->ValidateConsoleInputAndAddToCurrentPreset(CommittedText);
 	}
 
 	return FReply::Unhandled();
diff --git a/Engine/Plugins/Experimental/ConsoleVariables/Source/ConsoleVariablesEditor/Private/Views/MainPanel/SConsoleVariablesEditorMainPanel.cpp b/Engine/Plugins/Experimental/ConsoleVariables/Source/ConsoleVariablesEditor/Private/Views/MainPanel/SConsoleVariablesEditorMainPanel.cpp
--- a/Engine/Plugins/Experimental/ConsoleVariables/Source/ConsoleVariablesEditor/Private/Views/MainPanel/SConsoleVariablesEditorMainPanel.cpp
+++ b/Engine/Plugins/Experimental/ConsoleVariables/Source/ConsoleVariablesEditor/Private/Views/MainPanel/SConsoleVariablesEditorMainPanel.cpp
@@ -86,58 +86,9 @@ SConsoleVariablesEditorMainPanel::~SConsoleVariablesEditorMainPanel()
 	MultiUserDetailsBox.Reset();
 }
 
-FReply SConsoleVariablesEditorMainPanel::ValidateConsoleInput(const FText& CommittedText) const
+FReply SConsoleVariablesEditorMainPanel::ValidateConsoleInputAndAddToCurrentPreset(const FText& CommittedText) const
 {
-	const FString CommandString = CommittedText.ToString().TrimStartAndEnd();
-	FString CommandKey;
-	FString ValueString;
-	
-	if (CommandString.Contains(" "))
-	{
-		CommandString.Split(TEXT(" "), &CommandKey, &ValueString);
-	}
-	else
-	{
-		CommandKey = CommandString;
-	}
-
-	if (IConsoleObject* ConsoleObject = IConsoleManager::Get().FindConsoleObject(*CommandKey)) 
-	{
-		if (ValueString.IsEmpty())
-		{
-			if (const IConsoleVariable* AsVariable = ConsoleObject->AsVariable())
-			{
-				ValueString = AsVariable->GetString();
-			}
-		}
-		
-		MainPanel.Pin()->AddConsoleObjectToPreset(
-			CommandKey,
-			ValueString,
-			true
-		);
-	}
-	else if (CommandString.IsEmpty())
-	{
-		UE_LOG(LogConsoleVariablesEditor, Warning, TEXT("%hs: Input is blank."), __FUNCTION__);
-	}
-	// Try to execute the whole command. Some commands are not registered, but are parsed externally so they won't be found by IConsoleManager::FindConsoleObject
-	else if (GEngine->Exec(FConsoleVariablesEditorCommandInfo::GetCurrentWorld(), *CommandString))
-	{
-		MainPanel.Pin()->AddConsoleObjectToPreset(
-			CommandString,
-			"",
-			true
-			);
-	}
-	else
-	{
-		UE_LOG(LogConsoleVariablesEditor, Warning,
-			TEXT("%hs: Input '%s' is not a recognized console variable or command."),
-			__FUNCTION__, *CommandString);
-	}
-
-	return FReply::Handled();
+	return MainPanel.Pin()->ValidateConsoleInputAndAddToCurrentPreset(CommittedText);
 }
 
 void SConsoleVariablesEditorMainPanel::RefreshMultiUserDetails() const
@@ -293,7 +244,7 @@ TSharedRef<SWidget> SConsoleVariablesEditorMainPanel::OnGeneratePresetsMenu()
 		LOCTEXT("SavePreset_Text", "Save Preset"),
 		LOCTEXT("SavePreset_Tooltip", "Save the current preset if one has been loaded. Otherwise, the Save As dialog will be opened."),
 		FSlateIcon(FAppStyle::Get().GetStyleSetName(), "AssetEditor.SaveAsset"),
-		FUIAction(FExecuteAction::CreateRaw(MainPanel.Pin().Get(), &FConsoleVariablesEditorMainPanel::SavePreset)),
+		FUIAction(FExecuteAction::CreateRaw(MainPanel.Pin().Get(), &FConsoleVariablesEditorMainPanel::SaveCurrentPreset)),
 		NAME_None,
 		EUserInterfaceActionType::Button
 	);
@@ -302,7 +253,7 @@ TSharedRef<SWidget> SConsoleVariablesEditorMainPanel::OnGeneratePresetsMenu()
 		LOCTEXT("SavePresetAs_Text", "Save Preset As"),
 		LOCTEXT("SavePresetAs_Tooltip", "Save the current configuration as a new preset that can be shared between multiple jobs, or imported later as the base of a new configuration."),
 		FSlateIcon(FAppStyle::Get().GetStyleSetName(), "AssetEditor.SaveAssetAs"),
-		FUIAction(FExecuteAction::CreateRaw(MainPanel.Pin().Get(), &FConsoleVariablesEditorMainPanel::SavePresetAs)),
+		FUIAction(FExecuteAction::CreateRaw(MainPanel.Pin().Get(), &FConsoleVariablesEditorMainPanel::SaveCurrentPresetAs)),
 		NAME_None,
 		EUserInterfaceActionType::Button
 	);
diff --git a/Engine/Plugins/Experimental/ConsoleVariables/Source/ConsoleVariablesEditor/Private/Views/MainPanel/SConsoleVariablesEditorMainPanel.h b/Engine/Plugins/Experimental/ConsoleVariables/Source/ConsoleVariablesEditor/Private/Views/MainPanel/SConsoleVariablesEditorMainPanel.h
--- a/Engine/Plugins/Experimental/ConsoleVariables/Source/ConsoleVariablesEditor/Private/Views/MainPanel/SConsoleVariablesEditorMainPanel.h
+++ b/Engine/Plugins/Experimental/ConsoleVariables/Source/ConsoleVariablesEditor/Private/Views/MainPanel/SConsoleVariablesEditorMainPanel.h
@@ -27,7 +27,7 @@ public:
 	virtual ~SConsoleVariablesEditorMainPanel() override;
 
 	/** A handler for when text is committed to the Console Input Box that appears when clicking the "Add Console Variable" button */
-	FReply ValidateConsoleInput(const FText& CommittedText) const;
+	FReply ValidateConsoleInputAndAddToCurrentPreset(const FText& CommittedText) const;
 
 	void RefreshMultiUserDetails() const;
 	void ToggleMultiUserDetails(ECheckBoxState CheckState) const;
diff --git a/Engine/Plugins/Experimental/ConsoleVariables/Source/ConsoleVariablesEditor/Public/ConsoleVariablesEditorFunctionLibrary.h b/Engine/Plugins/Experimental/ConsoleVariables/Source/ConsoleVariablesEditor/Public/ConsoleVariablesEditorFunctionLibrary.h
new file mode 100644
--- /dev/null
+++ b/Engine/Plugins/Experimental/ConsoleVariables/Source/ConsoleVariablesEditor/Public/ConsoleVariablesEditorFunctionLibrary.h
@@ -0,0 +1,87 @@
+// Copyright Epic Games, Inc. All Rights Reserved.
+
+#pragma once
+
+#include "CoreMinimal.h"
+#include "Kismet/BlueprintFunctionLibrary.h"
+
+#include "ConsoleVariablesEditorFunctionLibrary.generated.h"
+
+class UConsoleVariablesAsset;
+
+/** An asset used to track collections of console variables that can be recalled and edited using the Console Variables Editor. */
+UCLASS(BlueprintType)
+class CONSOLEVARIABLESEDITOR_API UConsoleVariablesEditorFunctionLibrary : public UBlueprintFunctionLibrary
+{
+	GENERATED_BODY()
+public:
+
+	/** Return the currently loaded list of variables in the Console Variables Editor. */
+	UFUNCTION(BlueprintCallable, Category = "Console Variables Editor")
+	static UConsoleVariablesAsset* GetCurrentlyLoadedPreset();
+
+	/** Loads the given asset in the Console Variables Editor and sets all its variable values. */
+	UFUNCTION(BlueprintCallable, Category = "Console Variables Editor")
+	static void LoadPresetIntoConsoleVariablesEditor(const UConsoleVariablesAsset* InAsset);
+
+	/*
+	 * Saves the current list in the Console Variables Editor to the given asset.
+	 * The Asset will not be automatically saved.
+	 * @return true if successful.
+	 */
+	UFUNCTION(BlueprintCallable, Category = "Console Variables Editor")
+	static bool CopyCurrentListToAsset(UConsoleVariablesAsset* InAsset);
+
+	/*
+	 * Adds a validated command to the current preset with its current value.
+	 * The Asset will not be automatically saved.
+	 * @return true if successful.
+	 */
+	UFUNCTION(BlueprintCallable, Category = "Console Variables Editor")
+	static bool AddValidatedCommandToCurrentPreset(const FString NewCommand);
+
+	/*
+	 * Removes a command from the current preset if it exists in the saved data.
+	 * The Asset will not be automatically saved.
+	 * @return true if successful.
+	 */
+	UFUNCTION(BlueprintCallable, Category = "Console Variables Editor")
+	static bool RemoveCommandFromCurrentPreset(const FString NewCommand);
+
+	/** Return an array of strings containing the command names for each command found in the given preset. */
+	UFUNCTION(BlueprintCallable, Category = "Console Variables Editor")
+	static bool GetListOfCommandsFromPreset(const UConsoleVariablesAsset* InAsset, TArray<FString>& OutCommandList);
+
+	/** Set a console variable value directly. Returns true if the console object exists. */
+	UFUNCTION(BlueprintCallable, Category = "Console Variables Editor", meta=(DisplayName="Set Console Variable By Name (Float)"))
+	static bool SetConsoleVariableByName_Float(const FString InCommandName, const float InValue);
+
+	/** Set a console variable value directly. Returns true if the console object exists. */
+	UFUNCTION(BlueprintCallable, Category = "Console Variables Editor", meta=(DisplayName="Set Console Variable By Name (Bool)"))
+	static bool SetConsoleVariableByName_Bool(const FString InCommandName, const bool InValue);
+
+	/** Set a console variable value directly. Returns true if the console object exists. */
+	UFUNCTION(BlueprintCallable, Category = "Console Variables Editor", meta=(DisplayName="Set Console Variable By Name (Int)"))
+	static bool SetConsoleVariableByName_Int(const FString InCommandName, const int32 InValue);
+	
+	/** Set a console variable value directly. Returns true if the console object exists. */
+	UFUNCTION(BlueprintCallable, Category = "Console Variables Editor", meta=(DisplayName="Set Console Variable By Name (String)"))
+	static bool SetConsoleVariableByName_String(const FString InCommandName, const FString InValue);
+
+	/** Get a console variable's string value directly. Returns true if the console object exists. */
+	UFUNCTION(BlueprintCallable, Category = "Console Variables Editor")
+	static bool GetConsoleVariableStringValue(const FString InCommandName, FString& OutValue);
+
+	/** Set a console variable value directly. Returns true if the console object exists. */
+	UFUNCTION(BlueprintCallable, Category = "Console Variables Editor")
+	static bool GetConsoleVariableSourceByName(const FString InCommandName, FString& OutValue);
+
+	/** Return whether the Multi-user sync setting for the current instance of the editor is enabled. */
+	UFUNCTION(BlueprintCallable, Category = "Console Variables Editor")
+	static bool GetEnableMultiUserCVarSync();
+	
+	/** Enable or disable the Multi-user sync setting for the current instance of the editor. */
+	UFUNCTION(BlueprintCallable, Category = "Console Variables Editor")
+	static void SetEnableMultiUserCVarSync(bool bNewSetting);
+	
+};
