Console Variables Editor:

-Fix for: CVars are reset when another cvar is added or removed from the list

#[fyi] Jason.Walter

diff --git a/Engine/Plugins/Experimental/ConsoleVariables/Source/ConsoleVariablesEditor/Private/ConsoleVariablesEditorFunctionLibrary.cpp b/Engine/Plugins/Experimental/ConsoleVariables/Source/ConsoleVariablesEditor/Private/ConsoleVariablesEditorFunctionLibrary.cpp
--- a/Engine/Plugins/Experimental/ConsoleVariables/Source/ConsoleVariablesEditor/Private/ConsoleVariablesEditorFunctionLibrary.cpp
+++ b/Engine/Plugins/Experimental/ConsoleVariables/Source/ConsoleVariablesEditor/Private/ConsoleVariablesEditorFunctionLibrary.cpp
@@ -64,9 +64,7 @@ bool UConsoleVariablesEditorFunctionLibrary::AddValidatedCommandToCurrentPreset(
 
 bool UConsoleVariablesEditorFunctionLibrary::RemoveCommandFromCurrentPreset(const FString NewCommand)
 {
-	FConsoleVariablesEditorModule& ConsoleVariablesEditorModule = FConsoleVariablesEditorModule::Get();
-	
-	ConsoleVariablesEditorModule.UpdatePresetValuesForSave(GetCurrentlyLoadedPreset());
+	const FConsoleVariablesEditorModule& ConsoleVariablesEditorModule = FConsoleVariablesEditorModule::Get();
 	
 	if (const bool bSuccess = GetCurrentlyLoadedPreset()->RemoveConsoleVariable(NewCommand))
 	{
diff --git a/Engine/Plugins/Experimental/ConsoleVariables/Source/ConsoleVariablesEditor/Private/ConsoleVariablesEditorModule.cpp b/Engine/Plugins/Experimental/ConsoleVariables/Source/ConsoleVariablesEditor/Private/ConsoleVariablesEditorModule.cpp
--- a/Engine/Plugins/Experimental/ConsoleVariables/Source/ConsoleVariablesEditor/Private/ConsoleVariablesEditorModule.cpp
+++ b/Engine/Plugins/Experimental/ConsoleVariables/Source/ConsoleVariablesEditor/Private/ConsoleVariablesEditorModule.cpp
@@ -220,14 +220,14 @@ FReply FConsoleVariablesEditorModule::ValidateConsoleInputAndAddToCurrentPreset(
 	return MainPanel->ValidateConsoleInputAndAddToCurrentPreset(CommittedText);
 }
 
-void FConsoleVariablesEditorModule::RefreshList() const
+void FConsoleVariablesEditorModule::RebuildList(const FString InConsoleCommandToScrollTo, bool bShouldCacheValues) const
 {
-	MainPanel->RefreshList();
+	MainPanel->RebuildList(InConsoleCommandToScrollTo, bShouldCacheValues);
 }
 
-void FConsoleVariablesEditorModule::RebuildList() const
+void FConsoleVariablesEditorModule::RefreshList() const
 {
-	MainPanel->RebuildList();
+	MainPanel->RefreshList();
 }
 
 void FConsoleVariablesEditorModule::UpdatePresetValuesForSave(TObjectPtr<UConsoleVariablesAsset> InAsset)
diff --git a/Engine/Plugins/Experimental/ConsoleVariables/Source/ConsoleVariablesEditor/Private/ConsoleVariablesEditorModule.h b/Engine/Plugins/Experimental/ConsoleVariables/Source/ConsoleVariablesEditor/Private/ConsoleVariablesEditorModule.h
--- a/Engine/Plugins/Experimental/ConsoleVariables/Source/ConsoleVariablesEditor/Private/ConsoleVariablesEditorModule.h
+++ b/Engine/Plugins/Experimental/ConsoleVariables/Source/ConsoleVariablesEditor/Private/ConsoleVariablesEditorModule.h
@@ -67,8 +67,17 @@ public:
 	[[nodiscard]] TObjectPtr<UConsoleVariablesAsset> GetGlobalSearchAsset() const;
 	
 	FReply ValidateConsoleInputAndAddToCurrentPreset(const FText& CommittedText) const;
+	/*
+	 * Regenerate the list items and refresh the list. Call when adding or removing variables.
+	 * @param bShouldCacheValues If true, the current list's current values will be cached and then restored when the list is rebuilt. Otherwise preset values will be used.
+	 */
+	void RebuildList(const FString InConsoleCommandToScrollTo = "", bool bShouldCacheValues = true) const;
+	
+	/**
+	 * Refresh filters and sorting.
+	 * Useful for when the list state has gone stale but the variable count has not changed.
+	 */
 	void RefreshList() const;
-	void RebuildList() const;
 	
 	void UpdatePresetValuesForSave(TObjectPtr<UConsoleVariablesAsset> InAsset);
 
diff --git a/Engine/Plugins/Experimental/ConsoleVariables/Source/ConsoleVariablesEditor/Private/Views/List/ConsoleVariablesEditorList.cpp b/Engine/Plugins/Experimental/ConsoleVariables/Source/ConsoleVariablesEditor/Private/Views/List/ConsoleVariablesEditorList.cpp
--- a/Engine/Plugins/Experimental/ConsoleVariables/Source/ConsoleVariablesEditor/Private/Views/List/ConsoleVariablesEditorList.cpp
+++ b/Engine/Plugins/Experimental/ConsoleVariables/Source/ConsoleVariablesEditor/Private/Views/List/ConsoleVariablesEditorList.cpp
@@ -24,11 +24,11 @@ void FConsoleVariablesEditorList::SetSearchString(const FString& SearchString)
 	ListWidget->SetSearchStringInSearchInputField(SearchString);
 }
 
-void FConsoleVariablesEditorList::RebuildList(const FString& InConsoleCommandToScrollTo) const
+void FConsoleVariablesEditorList::RebuildList(const FString& InConsoleCommandToScrollTo, bool bShouldCacheValues) const
 {
 	if (ListWidget.IsValid())
 	{
-		ListWidget->RebuildList(InConsoleCommandToScrollTo);
+		ListWidget->RebuildList(InConsoleCommandToScrollTo, bShouldCacheValues);
 	}
 }
 
diff --git a/Engine/Plugins/Experimental/ConsoleVariables/Source/ConsoleVariablesEditor/Private/Views/List/ConsoleVariablesEditorList.h b/Engine/Plugins/Experimental/ConsoleVariables/Source/ConsoleVariablesEditor/Private/Views/List/ConsoleVariablesEditorList.h
--- a/Engine/Plugins/Experimental/ConsoleVariables/Source/ConsoleVariablesEditor/Private/Views/List/ConsoleVariablesEditorList.h
+++ b/Engine/Plugins/Experimental/ConsoleVariables/Source/ConsoleVariablesEditor/Private/Views/List/ConsoleVariablesEditorList.h
@@ -40,8 +40,11 @@ public:
 		CurrentListMode = NewListMode;;
 	}
 
-	/** Regenerate the list items and refresh the list. Call when adding or removing variables. */
-	void RebuildList(const FString& InConsoleCommandToScrollTo = "") const;
+	/*
+	 * Regenerate the list items and refresh the list. Call when adding or removing variables.
+	 * @param bShouldCacheValues If true, the current list's current values will be cached and then restored when the list is rebuilt. Otherwise preset values will be used.
+	 */
+	void RebuildList(const FString& InConsoleCommandToScrollTo = "", bool bShouldCacheValues = true) const;
 
 	/**
 	 * Refresh filters and sorting.
diff --git a/Engine/Plugins/Experimental/ConsoleVariables/Source/ConsoleVariablesEditor/Private/Views/List/ConsoleVariablesEditorListRow.cpp b/Engine/Plugins/Experimental/ConsoleVariables/Source/ConsoleVariablesEditor/Private/Views/List/ConsoleVariablesEditorListRow.cpp
--- a/Engine/Plugins/Experimental/ConsoleVariables/Source/ConsoleVariablesEditor/Private/Views/List/ConsoleVariablesEditorListRow.cpp
+++ b/Engine/Plugins/Experimental/ConsoleVariables/Source/ConsoleVariablesEditor/Private/Views/List/ConsoleVariablesEditorListRow.cpp
@@ -312,7 +312,7 @@ FReply FConsoleVariablesEditorListRow::OnActionButtonClicked()
 	const bool bIsGlobalSearch =
 		ListMode == FConsoleVariablesEditorList::EConsoleVariablesEditorListMode::GlobalSearch;
 
-	const FConsoleVariablesEditorModule& ConsoleVariablesEditorModule = FConsoleVariablesEditorModule::Get();
+	FConsoleVariablesEditorModule& ConsoleVariablesEditorModule = FConsoleVariablesEditorModule::Get();
 
 	const FString& CommandName = GetCommandInfo().Pin()->Command;
 
diff --git a/Engine/Plugins/Experimental/ConsoleVariables/Source/ConsoleVariablesEditor/Private/Views/List/SConsoleVariablesEditorList.cpp b/Engine/Plugins/Experimental/ConsoleVariables/Source/ConsoleVariablesEditor/Private/Views/List/SConsoleVariablesEditorList.cpp
--- a/Engine/Plugins/Experimental/ConsoleVariables/Source/ConsoleVariablesEditor/Private/Views/List/SConsoleVariablesEditorList.cpp
+++ b/Engine/Plugins/Experimental/ConsoleVariables/Source/ConsoleVariablesEditor/Private/Views/List/SConsoleVariablesEditorList.cpp
@@ -2,7 +2,6 @@
 
 #include "SConsoleVariablesEditorList.h"
 
-#include "ConsoleVariablesAsset.h"
 #include "ConsoleVariablesEditorList.h"
 #include "ConsoleVariablesEditorListFilters/ConsoleVariablesEditorListFilter_ModifiedVariables.h"
 #include "ConsoleVariablesEditorListFilters/ConsoleVariablesEditorListFilter_SourceText.h"
@@ -385,9 +384,20 @@ void SConsoleVariablesEditorList::RefreshGlobalSearchWidgets()
 	}
 }
 
-void SConsoleVariablesEditorList::RebuildList(const FString& InConsoleCommandToScrollTo)
+void SConsoleVariablesEditorList::RebuildList(const FString& InConsoleCommandToScrollTo, bool bShouldCacheValues)
 {
-	GenerateTreeView();
+	if (bShouldCacheValues)
+	{
+		CacheCurrentListItemData();
+	}
+
+	// Skip execution on load if we've cached the previous values
+	GenerateTreeView(!bShouldCacheValues);
+
+	if (bShouldCacheValues)
+	{
+		RestorePreviousListItemData();
+	}
 
 	RefreshList();
 
@@ -974,7 +984,61 @@ void SConsoleVariablesEditorList::OnListViewSearchTextChanged(const FText& Text)
 	ExecuteListViewSearchOnAllRows(Text.ToString(), false);
 }
 
-void SConsoleVariablesEditorList::GenerateTreeView()
+void SConsoleVariablesEditorList::CacheCurrentListItemData()
+{
+	CachedCommandStates.Empty();
+
+	// We only want preset items, not global search
+	const TArray<FConsoleVariablesEditorListRowPtr>& ListItems =
+		ListModelPtr.Pin()->GetListMode() == FConsoleVariablesEditorList::EConsoleVariablesEditorListMode::Preset ?
+			TreeViewRootObjects : LastPresetObjects;
+
+	for (const TSharedPtr<FConsoleVariablesEditorListRow>& Item : ListItems)
+	{
+		if (const TSharedPtr<FConsoleVariablesEditorCommandInfo>& CommandInfo = Item->GetCommandInfo().Pin())
+		{
+			if (const IConsoleVariable* AsVariable = CommandInfo->GetConsoleVariablePtr())
+			{
+				CachedCommandStates.Add(
+					{
+						CommandInfo->Command,
+						AsVariable->GetString(),
+						Item->GetWidgetCheckedState()
+					}
+				);
+			}
+		}
+	}
+}
+
+void SConsoleVariablesEditorList::RestorePreviousListItemData()
+{
+	for (const TSharedPtr<FConsoleVariablesEditorListRow>& Item : TreeViewRootObjects)
+	{
+		if (const TSharedPtr<FConsoleVariablesEditorCommandInfo>& CommandInfo = Item->GetCommandInfo().Pin())
+		{
+			if (CommandInfo->GetConsoleVariablePtr())
+			{
+				const FString& CommandName = CommandInfo->Command;
+				
+				if (const FConsoleVariablesEditorAssetSaveData* Match = Algo::FindByPredicate(
+					CachedCommandStates,
+					[CommandName](const FConsoleVariablesEditorAssetSaveData& CachedData)
+					{
+						return CachedData.CommandName.Equals(CommandName);
+					}))
+				{
+					CommandInfo->ExecuteCommand((*Match).CommandValueAsString);
+					Item->SetWidgetCheckedState((*Match).CheckedState);
+				}
+			}
+		}
+	}
+
+	CachedCommandStates.Empty();
+}
+
+void SConsoleVariablesEditorList::GenerateTreeView(const bool bExecuteCommandsAsTheyAreLoaded)
 {	
 	if (!ensure(TreeViewPtr.IsValid()))
 	{
@@ -1036,7 +1100,8 @@ void SConsoleVariablesEditorList::GenerateTreeView()
 					SavedCommand.CheckedState == ECheckBoxState::Unchecked ?
 							ECheckBoxState::Unchecked : ECheckBoxState::Checked;
 
-				if (ListModelPtr.Pin()->GetListMode() == FConsoleVariablesEditorList::EConsoleVariablesEditorListMode::Preset)
+				if (bExecuteCommandsAsTheyAreLoaded &&
+					ListModelPtr.Pin()->GetListMode() == FConsoleVariablesEditorList::EConsoleVariablesEditorListMode::Preset)
 				{
 					// If the row is checked and the saved value differs from the current value,
 					// execute the command with the saved value 
diff --git a/Engine/Plugins/Experimental/ConsoleVariables/Source/ConsoleVariablesEditor/Private/Views/List/SConsoleVariablesEditorList.h b/Engine/Plugins/Experimental/ConsoleVariables/Source/ConsoleVariablesEditor/Private/Views/List/SConsoleVariablesEditorList.h
--- a/Engine/Plugins/Experimental/ConsoleVariables/Source/ConsoleVariablesEditor/Private/Views/List/SConsoleVariablesEditorList.h
+++ b/Engine/Plugins/Experimental/ConsoleVariables/Source/ConsoleVariablesEditor/Private/Views/List/SConsoleVariablesEditorList.h
@@ -2,6 +2,7 @@
 
 #pragma once
 
+#include "ConsoleVariablesAsset.h"
 #include "ConsoleVariablesEditorListRow.h"
 #include "ConsoleVariablesEditorListFilters/IConsoleVariablesEditorListFilter.h"
 
@@ -44,8 +45,11 @@ public:
 	/** Remove all widgets from the global search container then recreate them */
 	void RefreshGlobalSearchWidgets();
 
-	/** Regenerate the list items and refresh the list. Call when adding or removing variables. */
-	void RebuildList(const FString& InConsoleCommandToScrollTo = "");
+	/*
+	 * Regenerate the list items and refresh the list. Call when adding or removing variables.
+	 * @param bShouldCacheValues If true, the current list's current values will be cached and then restored when the list is rebuilt. Otherwise preset values will be used.
+	 */
+	void RebuildList(const FString& InConsoleCommandToScrollTo = "", bool bShouldCacheValues = true);
 
 	/**
 	 * Refresh filters and sorting.
@@ -146,7 +150,9 @@ private:
 
 	//  Tree View Implementation
 
-	void GenerateTreeView();
+	void CacheCurrentListItemData();
+	void RestorePreviousListItemData();
+	void GenerateTreeView(const bool bExecuteCommandsAsTheyAreLoaded = true);
 	void FindVisibleTreeViewObjects();
 	void FindVisibleObjectsAndRequestTreeRefresh();
 	
@@ -168,6 +174,9 @@ private:
 	 */
 	TArray<FConsoleVariablesEditorListRowPtr> LastPresetObjects;
 
+	/** A collection of list item data to save and restore when rebuilding the list. */
+	TArray<FConsoleVariablesEditorAssetSaveData> CachedCommandStates;
+
 	// Sorting
 
 	FName ActiveSortingColumnName = NAME_None;
diff --git a/Engine/Plugins/Experimental/ConsoleVariables/Source/ConsoleVariablesEditor/Private/Views/List/SConsoleVariablesEditorListRow.cpp b/Engine/Plugins/Experimental/ConsoleVariables/Source/ConsoleVariablesEditor/Private/Views/List/SConsoleVariablesEditorListRow.cpp
--- a/Engine/Plugins/Experimental/ConsoleVariables/Source/ConsoleVariablesEditor/Private/Views/List/SConsoleVariablesEditorListRow.cpp
+++ b/Engine/Plugins/Experimental/ConsoleVariables/Source/ConsoleVariablesEditor/Private/Views/List/SConsoleVariablesEditorListRow.cpp
@@ -465,7 +465,8 @@ void SConsoleVariablesEditorListRow::OnCheckboxStateChange(const ECheckBoxState
 				TObjectPtr<UConsoleVariablesAsset> Asset =
 					ListModel.Pin()->GetListMode() == FConsoleVariablesEditorList::EConsoleVariablesEditorListMode::Preset ?
 						ConsoleVariablesEditorModule.GetPresetAsset() : ConsoleVariablesEditorModule.GetGlobalSearchAsset();
-				
+
+				// todo: Need to update JUST the
 				PinnedListView->UpdatePresetValuesForSave(Asset);
 			}
 		}
diff --git a/Engine/Plugins/Experimental/ConsoleVariables/Source/ConsoleVariablesEditor/Private/Views/MainPanel/ConsoleVariablesEditorMainPanel.cpp b/Engine/Plugins/Experimental/ConsoleVariables/Source/ConsoleVariablesEditor/Private/Views/MainPanel/ConsoleVariablesEditorMainPanel.cpp
--- a/Engine/Plugins/Experimental/ConsoleVariables/Source/ConsoleVariablesEditor/Private/Views/MainPanel/ConsoleVariablesEditorMainPanel.cpp
+++ b/Engine/Plugins/Experimental/ConsoleVariables/Source/ConsoleVariablesEditor/Private/Views/MainPanel/ConsoleVariablesEditorMainPanel.cpp
@@ -58,8 +58,6 @@ void FConsoleVariablesEditorMainPanel::AddConsoleObjectToCurrentPreset(
 {
 	if (const TObjectPtr<UConsoleVariablesAsset> Asset = GetEditingAsset())
 	{
-		UpdatePresetValuesForSave(Asset);
-		
 		Asset->AddOrSetConsoleObjectSavedData(
 			{
 				InConsoleCommand,
@@ -139,7 +137,7 @@ FReply FConsoleVariablesEditorMainPanel::ValidateConsoleInputAndAddToCurrentPres
 	return FReply::Handled();
 }
 
-void FConsoleVariablesEditorMainPanel::RebuildList(const FString InConsoleCommandToScrollTo) const
+void FConsoleVariablesEditorMainPanel::RebuildList(const FString InConsoleCommandToScrollTo, bool bShouldCacheValues) const
 {
 	if (EditorList.IsValid())
 	{
@@ -246,7 +244,7 @@ void FConsoleVariablesEditorMainPanel::ImportPreset(const FAssetData& InPresetAs
 	{
 		if (const TObjectPtr<UConsoleVariablesAsset> EditingAsset = GetEditingAsset(); ImportPreset_Impl(Preset, EditingAsset))
 		{
-			EditorList->RebuildList();
+			EditorList->RebuildList("", false);
 		}
 	}
 }
@@ -257,7 +255,7 @@ void FConsoleVariablesEditorMainPanel::ImportPreset(const TObjectPtr<UConsoleVar
 
 	if (const TObjectPtr<UConsoleVariablesAsset> EditingAsset = GetEditingAsset(); ImportPreset_Impl(InPreset, EditingAsset))
 	{
-		EditorList->RebuildList();
+		EditorList->RebuildList("", false);
 	}
 }
 
diff --git a/Engine/Plugins/Experimental/ConsoleVariables/Source/ConsoleVariablesEditor/Private/Views/MainPanel/ConsoleVariablesEditorMainPanel.h b/Engine/Plugins/Experimental/ConsoleVariables/Source/ConsoleVariablesEditor/Private/Views/MainPanel/ConsoleVariablesEditorMainPanel.h
--- a/Engine/Plugins/Experimental/ConsoleVariables/Source/ConsoleVariablesEditor/Private/Views/MainPanel/ConsoleVariablesEditorMainPanel.h
+++ b/Engine/Plugins/Experimental/ConsoleVariables/Source/ConsoleVariablesEditor/Private/Views/MainPanel/ConsoleVariablesEditorMainPanel.h
@@ -28,8 +28,11 @@ public:
 
 	FReply ValidateConsoleInputAndAddToCurrentPreset(const FText& CommittedText) const;
 
-	/** Regenerate the list items and refresh the list. Call when adding or removing variables. */
-	void RebuildList(const FString InConsoleCommandToScrollTo = "") const;
+	/*
+	 * Regenerate the list items and refresh the list. Call when adding or removing variables.
+	 * @param bShouldCacheValues If true, the current list's current values will be cached and then restored when the list is rebuilt. Otherwise preset values will be used.
+	 */
+	void RebuildList(const FString InConsoleCommandToScrollTo = "", bool bShouldCacheValues = true) const;
 
 	/**
 	 * Refresh filters and sorting.
