Console Variables Editor:

-Search failure copy changes
-Global Search pill padding adjustments
-Added fix for star icons disappearing on click in global search

#[fyi] Jason.Walter

diff --git a/Engine/Plugins/Experimental/ConsoleVariables/Source/ConsoleVariablesEditor/Private/Views/List/SConsoleVariablesEditorList.cpp b/Engine/Plugins/Experimental/ConsoleVariables/Source/ConsoleVariablesEditor/Private/Views/List/SConsoleVariablesEditorList.cpp
--- a/Engine/Plugins/Experimental/ConsoleVariables/Source/ConsoleVariablesEditor/Private/Views/List/SConsoleVariablesEditorList.cpp
+++ b/Engine/Plugins/Experimental/ConsoleVariables/Source/ConsoleVariablesEditor/Private/Views/List/SConsoleVariablesEditorList.cpp
@@ -21,6 +21,7 @@
 #include "Widgets/Input/SSearchBox.h"
 #include "Widgets/Layout/SWidgetSwitcher.h"
 #include "Widgets/Layout/SWrapBox.h"
+#include "Widgets/Text/SRichTextBlock.h"
 
 #define LOCTEXT_NAMESPACE "ConsoleVariablesEditor"
 
@@ -129,7 +130,7 @@ void SConsoleVariablesEditorList::Construct(const FArguments& InArgs, TSharedRef
 
 		+SVerticalBox::Slot()
 		.AutoHeight()
-		.Padding(FMargin(8.f, 0.f, 8.f, 0.f))
+		.Padding(FMargin(8.f, 4.f, 8.f, 4.f))
 		[
 			SAssignNew(GlobalSearchesHBox, SHorizontalBox)
 			.Visibility_Lambda([this]()
@@ -223,11 +224,23 @@ void SConsoleVariablesEditorList::Construct(const FArguments& InArgs, TSharedRef
 			.HAlign(HAlign_Fill)
 			.Padding(2.0f, 24.0f, 2.0f, 2.0f)
 			[
-				SNew(STextBlock)
+				SNew(SRichTextBlock)
 				.AutoWrapText(true)
 				.Justification(ETextJustify::Center)
-				.Text(LOCTEXT("ConsoleVariablesEditorList_NoList",
-					"No List to show. Try clearing the active local/global search and/or adding some console variables to the list."))
+				.Text_Lambda([this]()
+				{
+					if (const TSharedPtr<FConsoleVariablesEditorList> ListModel = ListModelPtr.Pin())
+					{
+						if (ListModel->GetListMode() == FConsoleVariablesEditorList::EConsoleVariablesEditorListMode::GlobalSearch)
+						{
+							return LOCTEXT("ConsoleVariablesEditorList_NoList",
+								"No matching console variables found in Unreal Engine.\n\nCheck your search criteria.");
+						}
+					}
+					
+					return LOCTEXT("ConsoleVariablesEditorList_NoList",
+					"No matching console variables in your list.\n\nCheck your filter or Search All console variables instead.");
+				})
 			]
 		]
 	];
diff --git a/Engine/Plugins/Experimental/ConsoleVariables/Source/ConsoleVariablesEditor/Private/Views/List/SConsoleVariablesEditorListRow.cpp b/Engine/Plugins/Experimental/ConsoleVariables/Source/ConsoleVariablesEditor/Private/Views/List/SConsoleVariablesEditorListRow.cpp
--- a/Engine/Plugins/Experimental/ConsoleVariables/Source/ConsoleVariablesEditor/Private/Views/List/SConsoleVariablesEditorListRow.cpp
+++ b/Engine/Plugins/Experimental/ConsoleVariables/Source/ConsoleVariablesEditor/Private/Views/List/SConsoleVariablesEditorListRow.cpp
@@ -458,17 +458,6 @@ void SConsoleVariablesEditorListRow::OnCheckboxStateChange(const ECheckBoxState
 		{
 			const TSharedPtr<SConsoleVariablesEditorList> PinnedListView = ListView.Pin();
 			PinnedListView->OnListItemCheckBoxStateChange(InNewState);
-
-			if (const TWeakPtr<FConsoleVariablesEditorList>& ListModel = PinnedListView->GetListModelPtr(); ListModel.IsValid())
-			{
-				const FConsoleVariablesEditorModule& ConsoleVariablesEditorModule = FConsoleVariablesEditorModule::Get();
-				TObjectPtr<UConsoleVariablesAsset> Asset =
-					ListModel.Pin()->GetListMode() == FConsoleVariablesEditorList::EConsoleVariablesEditorListMode::Preset ?
-						ConsoleVariablesEditorModule.GetPresetAsset() : ConsoleVariablesEditorModule.GetGlobalSearchAsset();
-
-				// todo: Need to update JUST the
-				PinnedListView->UpdatePresetValuesForSave(Asset);
-			}
 		}
 	}
 }
@@ -583,7 +572,10 @@ void SConsoleVariablesEditorListRowHoverWidgets::Construct(const FArguments& InA
 		.OnClicked_Lambda([this]()
 		{
 			SetVisibility(EVisibility::Collapsed);
-			return Item.Pin()->OnActionButtonClicked();
+			FReply Reply = Item.Pin()->OnActionButtonClicked();
+			DetermineButtonImageAndTooltip();
+			SetVisibility(EVisibility::SelfHitTestInvisible);
+			return Reply;
 		})
 		[
 			SNew(SScaleBox)
