Console Variables Editor:

-Search failure copy bolding

#[fyi] Jason.Walter

diff --git a/Engine/Plugins/Experimental/ConsoleVariables/Source/ConsoleVariablesEditor/Private/Views/List/SConsoleVariablesEditorList.cpp b/Engine/Plugins/Experimental/ConsoleVariables/Source/ConsoleVariablesEditor/Private/Views/List/SConsoleVariablesEditorList.cpp
--- a/Engine/Plugins/Experimental/ConsoleVariables/Source/ConsoleVariablesEditor/Private/Views/List/SConsoleVariablesEditorList.cpp
+++ b/Engine/Plugins/Experimental/ConsoleVariables/Source/ConsoleVariablesEditor/Private/Views/List/SConsoleVariablesEditorList.cpp
@@ -225,6 +225,7 @@ void SConsoleVariablesEditorList::Construct(const FArguments& InArgs, TSharedRef
 			.Padding(2.0f, 24.0f, 2.0f, 2.0f)
 			[
 				SNew(SRichTextBlock)
+				.DecoratorStyleSet(&FEditorStyle::Get())
 				.AutoWrapText(true)
 				.Justification(ETextJustify::Center)
 				.Text_Lambda([this]()
@@ -239,7 +240,7 @@ void SConsoleVariablesEditorList::Construct(const FArguments& InArgs, TSharedRef
 					}
 					
 					return LOCTEXT("ConsoleVariablesEditorList_NoList",
-					"No matching console variables in your list.\n\nCheck your filter or Search All console variables instead.");
+					"No matching console variables in your list.\n\nCheck your filter or <RichTextBlock.Bold>Search All</> console variables instead.");
 				})
 			]
 		]
