Use FCObject to track UObjects in the main Console Variable editor. Also make sure we sync console variables to nDisplay using unique ids so they can be synchronized.

#jria UE-137314

diff --git a/Engine/Plugins/Experimental/ConsoleVariables/Source/ConsoleVariablesEditor/Private/ConsoleVariablesEditorModule.cpp b/Engine/Plugins/Experimental/ConsoleVariables/Source/ConsoleVariablesEditor/Private/ConsoleVariablesEditorModule.cpp
--- a/Engine/Plugins/Experimental/ConsoleVariables/Source/ConsoleVariablesEditor/Private/ConsoleVariablesEditorModule.cpp
+++ b/Engine/Plugins/Experimental/ConsoleVariables/Source/ConsoleVariablesEditor/Private/ConsoleVariablesEditorModule.cpp
@@ -436,6 +436,17 @@ void FConsoleVariablesEditorModule::OpenConsoleVariablesEditor()
 	FGlobalTabmanager::Get()->TryInvokeTab(ConsoleVariablesToolkitPanelTabId);
 }
 
+void FConsoleVariablesEditorModule::AddReferencedObjects(FReferenceCollector& Collector)
+{
+	Collector.AddReferencedObject(EditingPresetAsset);
+	Collector.AddReferencedObject(EditingGlobalSearchAsset);
+}
+
+FString FConsoleVariablesEditorModule::GetReferencerName() const
+{
+	return TEXT("FConsoleVariablesEditorModule");
+}
+
 #undef LOCTEXT_NAMESPACE
-	
+
 IMPLEMENT_MODULE(FConsoleVariablesEditorModule, ConsoleVariablesEditor)
diff --git a/Engine/Plugins/Experimental/ConsoleVariables/Source/ConsoleVariablesEditor/Private/ConsoleVariablesEditorModule.h b/Engine/Plugins/Experimental/ConsoleVariables/Source/ConsoleVariablesEditor/Private/ConsoleVariablesEditorModule.h
--- a/Engine/Plugins/Experimental/ConsoleVariables/Source/ConsoleVariablesEditor/Private/ConsoleVariablesEditorModule.h
+++ b/Engine/Plugins/Experimental/ConsoleVariables/Source/ConsoleVariablesEditor/Private/ConsoleVariablesEditorModule.h
@@ -17,7 +17,7 @@ class FConsoleVariablesEditorMainPanel;
 class FConsoleVariablesEditorToolkit;
 class ISettingsSection;
 
-class FConsoleVariablesEditorModule : public IModuleInterface
+class FConsoleVariablesEditorModule : public IModuleInterface, public FGCObject
 {
 public:
 	
@@ -86,7 +86,10 @@ public:
 
 	void SendMultiUserConsoleVariableChange(const FString& InVariableName, const FString& InValueAsString) const;
 	void OnRemoteCvarChanged(const FString InName, const FString InValue);
-	
+
+	virtual void AddReferencedObjects( FReferenceCollector& Collector )  override;
+	virtual FString GetReferencerName() const override;
+
 private:
 
 	void OnFEngineLoopInitComplete();
diff --git a/Engine/Plugins/Experimental/ConsoleVariables/Source/ConsoleVariablesEditor/Private/MultiUser/ConsoleVariableSync.cpp b/Engine/Plugins/Experimental/ConsoleVariables/Source/ConsoleVariablesEditor/Private/MultiUser/ConsoleVariableSync.cpp
--- a/Engine/Plugins/Experimental/ConsoleVariables/Source/ConsoleVariablesEditor/Private/MultiUser/ConsoleVariableSync.cpp
+++ b/Engine/Plugins/Experimental/ConsoleVariables/Source/ConsoleVariablesEditor/Private/MultiUser/ConsoleVariableSync.cpp
@@ -402,7 +402,8 @@ struct FManagerImpl
 		}
 
 		FConcertSetConsoleVariableEvent OutEvent{MoveTemp(InName), MoveTemp(InValue)};
-		Session->SendCustomEvent(OutEvent, Session->GetSessionClientEndpointIds(), EConcertMessageFlags::ReliableOrdered);
+		Session->SendCustomEvent(OutEvent, Session->GetSessionClientEndpointIds(),
+								 EConcertMessageFlags::ReliableOrdered | EConcertMessageFlags::UniqueId);
 	}
 
 	/** Reference to the customization object so that we can sync with the session.*/
